Fix non-unity compile errors in Lightmass and CrashReportClient on Mac.

diff --git a/Engine/Source/Developer/CrashReportHelper/Private/Mac/MacErrorReport.cpp b/Engine/Source/Developer/CrashReportHelper/Private/Mac/MacErrorReport.cpp
--- a/Engine/Source/Developer/CrashReportHelper/Private/Mac/MacErrorReport.cpp
+++ b/Engine/Source/Developer/CrashReportHelper/Private/Mac/MacErrorReport.cpp
@@ -7,6 +7,8 @@
 #include "CrashDebugHelper.h"
 #include "Misc/FileHelper.h"
 #include "HAL/PlatformFilemanager.h"
+#include "Modules/ModuleManager.h"
+#include "CrashDebugHelperModule.h"
 
 #define LOCTEXT_NAMESPACE "CrashReportClient"
 
diff --git a/Engine/Source/Developer/IOS/IOSTargetPlatform/Private/Mac/IOSDeviceHelperMac.cpp b/Engine/Source/Developer/IOS/IOSTargetPlatform/Private/Mac/IOSDeviceHelperMac.cpp
--- a/Engine/Source/Developer/IOS/IOSTargetPlatform/Private/Mac/IOSDeviceHelperMac.cpp
+++ b/Engine/Source/Developer/IOS/IOSTargetPlatform/Private/Mac/IOSDeviceHelperMac.cpp
@@ -1,6 +1,7 @@
 // Copyright 1998-2018 Epic Games, Inc. All Rights Reserved
 
 #include "IOSTargetPlatform.h"
+#include "HAL/PlatformFile.h"
 
 #include <CoreFoundation/CoreFoundation.h>
 
diff --git a/Engine/Source/Programs/CrashReportClient/Private/Mac/CrashReportClientMainMac.cpp b/Engine/Source/Programs/CrashReportClient/Private/Mac/CrashReportClientMainMac.cpp
--- a/Engine/Source/Programs/CrashReportClient/Private/Mac/CrashReportClientMainMac.cpp
+++ b/Engine/Source/Programs/CrashReportClient/Private/Mac/CrashReportClientMainMac.cpp
@@ -3,6 +3,7 @@
 #include "CrashReportClientApp.h"
 #include "HAL/ExceptionHandling.h"
 #include "Mac/MacPlatformCrashContext.h"
+#include "Mac/MacPlatformProcess.h"
 #include "Mac/CocoaThread.h"
 #include "HAL/PlatformApplicationMisc.h"
 
diff --git a/Engine/Source/Programs/CrashReportClient/Private/Mac/MacErrorReport.cpp b/Engine/Source/Programs/CrashReportClient/Private/Mac/MacErrorReport.cpp
--- a/Engine/Source/Programs/CrashReportClient/Private/Mac/MacErrorReport.cpp
+++ b/Engine/Source/Programs/CrashReportClient/Private/Mac/MacErrorReport.cpp
@@ -7,6 +7,8 @@
 #include "CrashDebugHelper.h"
 #include "Misc/FileHelper.h"
 #include "HAL/PlatformFilemanager.h"
+#include "Modules/ModuleManager.h"
+#include "CrashDebugHelperModule.h"
 
 #define LOCTEXT_NAMESPACE "CrashReportClient"
 
diff --git a/Engine/Source/Programs/UnrealLightmass/Private/LightmassCore/Math/LMkDOP.h b/Engine/Source/Programs/UnrealLightmass/Private/LightmassCore/Math/LMkDOP.h
--- a/Engine/Source/Programs/UnrealLightmass/Private/LightmassCore/Math/LMkDOP.h
+++ b/Engine/Source/Programs/UnrealLightmass/Private/LightmassCore/Math/LMkDOP.h
@@ -6,6 +6,8 @@
 // #include "Collision.h"
 #include "LMMath.h"
 #include "LMMathSSE.h"
+#include "LMStats.h"
+#include "UnrealLightmass.h"
 
 namespace Lightmass
 {
diff --git a/Engine/Source/Programs/UnrealLightmass/Private/LightmassCore/Misc/LMThreading.cpp b/Engine/Source/Programs/UnrealLightmass/Private/LightmassCore/Misc/LMThreading.cpp
--- a/Engine/Source/Programs/UnrealLightmass/Private/LightmassCore/Misc/LMThreading.cpp
+++ b/Engine/Source/Programs/UnrealLightmass/Private/LightmassCore/Misc/LMThreading.cpp
@@ -6,6 +6,7 @@
 
 #include "LMThreading.h"
 #include "LMCore.h"
+#include "HAL/PlatformProcess.h"
 
 #if PLATFORM_WINDOWS
 #include "Windows/WindowsHWrapper.h"
