Fix non-unity compile errors
#rnx

diff --git a/Engine/Source/Programs/CrashReportClient/Private/CrashReportClientUnattended.cpp b/Engine/Source/Programs/CrashReportClient/Private/CrashReportClientUnattended.cpp
--- a/Engine/Source/Programs/CrashReportClient/Private/CrashReportClientUnattended.cpp
+++ b/Engine/Source/Programs/CrashReportClient/Private/CrashReportClientUnattended.cpp
@@ -4,6 +4,7 @@
 #include "Containers/Ticker.h"
 #include "CrashReportClientConfig.h"
 #include "CrashDescription.h"
+#include "Stats/Stats.h"
 
 FCrashReportClientUnattended::FCrashReportClientUnattended(FPlatformErrorReport& InErrorReport)
 	: ReceiverUploader(FCrashReportClientConfig::Get().GetReceiverAddress())
diff --git a/Engine/Source/Programs/CrashReportClient/Private/CrashUpload.cpp b/Engine/Source/Programs/CrashReportClient/Private/CrashUpload.cpp
--- a/Engine/Source/Programs/CrashReportClient/Private/CrashUpload.cpp
+++ b/Engine/Source/Programs/CrashReportClient/Private/CrashUpload.cpp
@@ -18,6 +18,8 @@
 #include "Misc/EngineBuildSettings.h"
 #include "Misc/CommandLine.h"
 
+#include "Stats/Stats.h"
+
 // Switched off CRR upload - Jun 2016
 #define PRIMARY_UPLOAD_RECEIVER 0
 #define PRIMARY_UPLOAD_DATAROUTER 1
