Add a pop up message box for Linux to check if we want to or not submit crash report data

diff --git a/Engine/Source/Programs/CrashReportClient/Private/CrashReportClientApp.cpp b/Engine/Source/Programs/CrashReportClient/Private/CrashReportClientApp.cpp
--- a/Engine/Source/Programs/CrashReportClient/Private/CrashReportClientApp.cpp
+++ b/Engine/Source/Programs/CrashReportClient/Private/CrashReportClientApp.cpp
@@ -337,8 +337,19 @@ void RunCrashReportClient(const TCHAR* CommandLine)
 
 	FPlatformErrorReport::Init();
 	FPlatformErrorReport ErrorReport = LoadErrorReport();
-	
-	if (ErrorReport.HasFilesToUpload() && FPrimaryCrashProperties::Get() != nullptr)
+
+	// For now Linux needs to ask if we want to submit a crash report unless specifically configured not to
+	if (PLATFORM_LINUX)
+	{
+		bool bAgreeToServerCrashUpload = true;
+		GConfig->GetBool(TEXT("CrashReportClient"), TEXT("bAgreeToServerCrashUpload"), bAgreeToServerCrashUpload, GEngineIni);
+		if(bAgreeToServerCrashUpload && FPlatformMisc::MessageBoxExt(EAppMsgType::YesNo, TEXT("A crash has occurred. Would you like to submit this crash report?"), TEXT("Crash Client Reporter")) == EAppReturnType::No)
+		{
+			GIsRequestingExit = true;
+		}
+	}
+
+	if (!GIsRequestingExit && ErrorReport.HasFilesToUpload() && FPrimaryCrashProperties::Get() != nullptr)
 	{
 		ErrorReport.SetCrashReportClientVersion(FCrashReportClientConfig::Get().GetVersion());
 
