Make sure the CrashReportClient background thread finishes work before we start shutting down the application. Should have been happening in the destructor, but something (Slate?) is keeping a reference to the CRC instance.

diff --git a/Engine/Source/Programs/CrashReportClient/Private/CrashReportClient.cpp b/Engine/Source/Programs/CrashReportClient/Private/CrashReportClient.cpp
--- a/Engine/Source/Programs/CrashReportClient/Private/CrashReportClient.cpp
+++ b/Engine/Source/Programs/CrashReportClient/Private/CrashReportClient.cpp
@@ -40,42 +40,48 @@ FCrashReportClient::FCrashReportClient(const FPlatformErrorReport& InErrorReport
 {
 	if (FPrimaryCrashProperties::Get()->IsValid())
 	{
-	bool bUsePrimaryData = false;
-	if (FPrimaryCrashProperties::Get()->HasProcessedData())
-	{
-		bUsePrimaryData = true;
-	}
-	else
-	{
-		if (!ErrorReport.TryReadDiagnosticsFile() && !FParse::Param( FCommandLine::Get(), TEXT( "no-local-diagnosis" ) ))
+		bool bUsePrimaryData = false;
+		if (FPrimaryCrashProperties::Get()->HasProcessedData())
 		{
-			DiagnoseReportTask = new FAsyncTask<FDiagnoseReportWorker>( this );
-			DiagnoseReportTask->StartBackgroundTask();
+			bUsePrimaryData = true;
 		}
 		else
 		{
-			bUsePrimaryData = true;
+			if (!ErrorReport.TryReadDiagnosticsFile() && !FParse::Param( FCommandLine::Get(), TEXT( "no-local-diagnosis" ) ))
+			{
+				DiagnoseReportTask = new FAsyncTask<FDiagnoseReportWorker>( this );
+				DiagnoseReportTask->StartBackgroundTask();
+			}
+			else
+			{
+				bUsePrimaryData = true;
+			}
 		}
-	}
 
-	if (bUsePrimaryData)
-	{
-		const FString CallstackString = FPrimaryCrashProperties::Get()->CallStack.AsString();
-		const FString ReportString = FString::Printf( TEXT( "%s\n\n%s" ), *FPrimaryCrashProperties::Get()->ErrorMessage.AsString(), *CallstackString );
-		DiagnosticText = FText::FromString( ReportString );
+		if (bUsePrimaryData)
+		{
+			const FString CallstackString = FPrimaryCrashProperties::Get()->CallStack.AsString();
+			const FString ReportString = FString::Printf( TEXT( "%s\n\n%s" ), *FPrimaryCrashProperties::Get()->ErrorMessage.AsString(), *CallstackString );
+			DiagnosticText = FText::FromString( ReportString );
 
-		FormattedDiagnosticText = FCrashReportUtil::FormatDiagnosticText( FText::FromString( ReportString ) );
-	}
+			FormattedDiagnosticText = FCrashReportUtil::FormatDiagnosticText( FText::FromString( ReportString ) );
+		}
 	}
 }
 
 
 FCrashReportClient::~FCrashReportClient()
 {
-	if( DiagnoseReportTask )
+	StopBackgroundThread();
+}
+
+void FCrashReportClient::StopBackgroundThread()
+{
+	if (DiagnoseReportTask)
 	{
 		DiagnoseReportTask->EnsureCompletion();
 		delete DiagnoseReportTask;
+		DiagnoseReportTask = nullptr;
 	}
 }
 
diff --git a/Engine/Source/Programs/CrashReportClient/Private/CrashReportClient.h b/Engine/Source/Programs/CrashReportClient/Private/CrashReportClient.h
--- a/Engine/Source/Programs/CrashReportClient/Private/CrashReportClient.h
+++ b/Engine/Source/Programs/CrashReportClient/Private/CrashReportClient.h
@@ -67,6 +67,9 @@ public:
 	/** Destructor. */
 	virtual ~FCrashReportClient();
 
+	/** Stops processing work in the background. */
+	void StopBackgroundThread();
+
 	/** Closes the crash report client without sending any data. Except the startup analytics. */
 	FReply CloseWithoutSending();
 
diff --git a/Engine/Source/Programs/CrashReportClient/Private/CrashReportClientApp.cpp b/Engine/Source/Programs/CrashReportClient/Private/CrashReportClientApp.cpp
--- a/Engine/Source/Programs/CrashReportClient/Private/CrashReportClientApp.cpp
+++ b/Engine/Source/Programs/CrashReportClient/Private/CrashReportClientApp.cpp
@@ -245,7 +245,7 @@ bool RunWithUI(FPlatformErrorReport ErrorReport)
 	// open up the app window	
 	TSharedRef<SCrashReportClient> ClientControl = SNew(SCrashReportClient, CrashReportClient);
 
-	auto Window = FSlateApplication::Get().AddWindow(
+	TSharedRef<SWindow> Window = FSlateApplication::Get().AddWindow(
 		SNew(SWindow)
 		.Title(NSLOCTEXT("CrashReportClient", "CrashReportClientAppName", "Unreal Engine 4 Crash Reporter"))
 		.HasCloseButton(FCrashReportClientConfig::Get().IsAllowedToCloseWithoutSending())
@@ -276,6 +276,12 @@ bool RunWithUI(FPlatformErrorReport ErrorReport)
 		}
 	}
 
+	// Make sure the window is hidden, because it might take a while for the background thread to finish.
+	Window->HideWindow();
+
+	// Stop the background thread
+	CrashReportClient->StopBackgroundThread();
+
 	// Clean up the custom styles
 	FCrashReportClientStyle::Shutdown();
 
