Exposing CrashType and EngineData to the CRC analytic event.

diff --git a/Engine/Source/Programs/CrashReportClient/Private/CrashDescription.cpp b/Engine/Source/Programs/CrashReportClient/Private/CrashDescription.cpp
--- a/Engine/Source/Programs/CrashReportClient/Private/CrashDescription.cpp
+++ b/Engine/Source/Programs/CrashReportClient/Private/CrashDescription.cpp
@@ -282,6 +282,7 @@ void FPrimaryCrashProperties::SendPreUploadAnalytics()
   * @EventParam CrashVersion" - Describes the version of the crash data pipeline we used on the client side (1 = Legacy metadata based, 2 = unused, 3 = New crash context based)
   * @EventParam CrashGUID - GUID for this event generated by the crashed process (e.g. Windows format is UE4CC-Windows-1F07494140C2669B52AC96A1C4D20F83_0000 with the last four digits representing the index of event within a single session)
   * @EventParam PlatformCallbackResult - Platform-specific UE4 Core value. (integer)
+  * @EventParam CrashType - The name of the crash report type (e.g Assert, Ensure, Hang)
   * @EventParam GameName - The name of the game that crashed from FApp::GetProjectName(). (e.g. UE4-QAGame)
   * @EventParam EngineVersion - The engine version of the crashed process from FEngineVersion::Current(). (e.g. 4.13.1-3142249+++UE4+Release-4.13)
   * @EventParam CrashReportClientVersion - CrashReportClient version, set in its config. (e.g. "1.0.0")
@@ -317,6 +318,7 @@ void SendPostUploadEnsureAnalytics(const TArray<FAnalyticsEventAttribute>& InCra
  * @EventParam CrashVersion" - Describes the version of the crash data pipeline we used on the client side (1 = Legacy metadata based, 2 = unused, 3 = New crash context based)
  * @EventParam CrashGUID - GUID for this event generated by the crashed process (e.g. Windows format is UE4CC-Windows-1F07494140C2669B52AC96A1C4D20F83_0000 with the last four digits representing the index of event within a single session)
  * @EventParam PlatformCallbackResult - Platform-specific UE4 Core value. (integer)
+ * @EventParam CrashType - The name of the crash report type (e.g Assert, Ensure, Hang)
  * @EventParam GameName - The name of the game that crashed from FApp::GetProjectName(). (e.g. UE4-QAGame)
  * @EventParam EngineVersion - The engine version of the crashed process from FEngineVersion::Current(). (e.g. 4.13.1-3142249+++UE4+Release-4.13)
  * @EventParam CrashReportClientVersion - CrashReportClient version, set in its config. (e.g. "1.0.0")
@@ -364,6 +366,7 @@ void FPrimaryCrashProperties::MakeCrashEventAttributes(TArray<FAnalyticsEventAtt
 	OutCrashAttributes.Add(FAnalyticsEventAttribute(TEXT("CrashVersion"), (int32)CrashVersion));
 	OutCrashAttributes.Add(FAnalyticsEventAttribute(TEXT("CrashGUID"), CrashGUID));
 	OutCrashAttributes.Add(FAnalyticsEventAttribute(TEXT("PlatformCallbackResult"), PlatformCallbackResult.AsString()));
+	OutCrashAttributes.Add(FAnalyticsEventAttribute(TEXT("CrashType"), CrashType));
 
 	//	AppID = GameName
 	OutCrashAttributes.Add(FAnalyticsEventAttribute(TEXT("GameName"), GameName));
@@ -388,11 +391,22 @@ void FPrimaryCrashProperties::MakeCrashEventAttributes(TArray<FAnalyticsEventAtt
 	OutCrashAttributes.Add(FAnalyticsEventAttribute(TEXT("GameSessionID"), GameSessionID.AsString()));
 	OutCrashAttributes.Add(FAnalyticsEventAttribute(TEXT("DeploymentName"), DeploymentName));
 
+	// Add arbitrary engine data
+	const FXmlNode* EngineNode = XmlFile->GetRootNode()->FindChildNode( FGenericCrashContext::EngineDataTag );
+	if (EngineNode)
+	{
+		for (const FXmlNode* ChildNode : EngineNode->GetChildrenNodes())
+		{
+			FString KeyName = FString(TEXT("EngineData.")) + ChildNode->GetTag();
+			OutCrashAttributes.Add(FAnalyticsEventAttribute(*KeyName, ChildNode->GetContent()));
+		}
+	}
+
 	// Add arbitrary game data
-	const FXmlNode* MainNode = XmlFile->GetRootNode()->FindChildNode( FGenericCrashContext::GameDataTag );
-	if (MainNode)
+	const FXmlNode* GameNode = XmlFile->GetRootNode()->FindChildNode( FGenericCrashContext::GameDataTag );
+	if (GameNode)
 	{
-		for (const FXmlNode* ChildNode : MainNode->GetChildrenNodes())
+		for (const FXmlNode* ChildNode : GameNode->GetChildrenNodes())
 		{
 			FString KeyName = FString(TEXT("GameData.")) + ChildNode->GetTag();
 			OutCrashAttributes.Add(FAnalyticsEventAttribute(*KeyName, ChildNode->GetContent()));
@@ -453,6 +467,7 @@ FCrashContext::FCrashContext( const FString& CrashContextFilepath )
 		GetCrashProperty( DeploymentName, FGenericCrashContext::RuntimePropertiesTag, TEXT( "DeploymentName" ) );
 		GetCrashProperty( AppDefaultLocale, FGenericCrashContext::RuntimePropertiesTag, TEXT( "AppDefaultLocale" ) );
 		GetCrashProperty( bIsEnsure, FGenericCrashContext::RuntimePropertiesTag, TEXT("IsEnsure"));
+		GetCrashProperty( CrashType, FGenericCrashContext::RuntimePropertiesTag, TEXT("CrashType"));
 		GetCrashProperty( NumMinidumpFramesToIgnore, FGenericCrashContext::RuntimePropertiesTag, TEXT("NumMinidumpFramesToIgnore"));
 
 		if (CrashDumpMode == ECrashDumpMode::FullDump)
diff --git a/Engine/Source/Programs/CrashReportClient/Private/CrashDescription.h b/Engine/Source/Programs/CrashReportClient/Private/CrashDescription.h
--- a/Engine/Source/Programs/CrashReportClient/Private/CrashDescription.h
+++ b/Engine/Source/Programs/CrashReportClient/Private/CrashDescription.h
@@ -371,6 +371,11 @@ struct FPrimaryCrashProperties
 	 */
 	bool bIsEnsure;
 
+	/**
+	 *	The type of crash being reported, e.g. Assert, Ensure, Hang
+	 */
+	FString CrashType;
+
 protected:
 	/** Default constructor. */
 	FPrimaryCrashProperties();
