Use the CRC to pop a message up not in crash reporting code

diff --git a/Engine/Source/Programs/CrashReportClient/Private/CrashReportClientApp.cpp b/Engine/Source/Programs/CrashReportClient/Private/CrashReportClientApp.cpp
--- a/Engine/Source/Programs/CrashReportClient/Private/CrashReportClientApp.cpp
+++ b/Engine/Source/Programs/CrashReportClient/Private/CrashReportClientApp.cpp
@@ -54,6 +54,9 @@ static FString GameNameFromCmd;
 /** GUID of the crash passed via the command line. */
 static FString CrashGUIDFromCmd;
 
+/** If we are implicitly sending its assumed we are also unattended for now */
+static bool bImplicitSendFromCmd = false;
+
 /**
  * Look for the report to upload, either in the command line or in the platform's report queue
  */
@@ -107,6 +110,11 @@ void ParseCommandLine(const TCHAR* CommandLine)
 		{
 			CrashGUIDFromCmd = Params.FindRef(TEXT("CrashGUID"));
 		}
+
+		if (Switches.Contains(TEXT("ImplicitSend")))
+		{
+			bImplicitSendFromCmd = true;
+		}
 	}
 
 	if (FoundReportDirectoryAbsolutePaths.Num() == 0)
@@ -292,6 +300,23 @@ bool RunWithUI(FPlatformErrorReport ErrorReport)
 }
 #endif // !CRASH_REPORT_UNATTENDED_ONLY
 
+// When we want to implicitly send and use unattended we still want to show a message box of a crash if possible
+class FMessageBoxThread : public FRunnable
+{
+	virtual uint32 Run() override
+	{
+		// We will not have any GUI for the crash reporter if we are sending implicitly, so pop a message box up at least
+		if (FApp::CanEverRender())
+		{
+			FPlatformMisc::MessageBoxExt(EAppMsgType::Ok,
+				*NSLOCTEXT("MessageDialog", "ReportCrash_Body", "The application has crashed and will now close. We apologize for the inconvenience.").ToString(),
+				*NSLOCTEXT("MessageDialog", "ReportCrash_Title", "Application Crash Detected").ToString());
+		}
+
+		return 0;
+	}
+};
+
 void RunUnattended(FPlatformErrorReport ErrorReport)
 {
 	// Set up the main ticker
@@ -301,11 +326,24 @@ void RunUnattended(FPlatformErrorReport ErrorReport)
 	FCrashReportCoreUnattended CrashReportClient(ErrorReport);
 	ErrorReport.SetUserComment(NSLOCTEXT("CrashReportClient", "UnattendedMode", "Sent in the unattended mode"));
 
+	FMessageBoxThread MessageBox;
+	FRunnableThread* MessageBoxThread = nullptr;
+
+	if (bImplicitSendFromCmd)
+	{
+		MessageBoxThread = FRunnableThread::Create(&MessageBox, TEXT("CrashReporter_MessageBox"));
+	}
+
 	// loop until the app is ready to quit
 	while (!GIsRequestingExit)
 	{
 		MainLoop.Tick();
 	}
+
+	if (bImplicitSendFromCmd && MessageBoxThread)
+	{
+		MessageBoxThread->WaitForCompletion();
+	}
 }
 
 
diff --git a/Engine/Source/Runtime/Core/Private/Mac/MacPlatformMisc.cpp b/Engine/Source/Runtime/Core/Private/Mac/MacPlatformMisc.cpp
--- a/Engine/Source/Runtime/Core/Private/Mac/MacPlatformMisc.cpp
+++ b/Engine/Source/Runtime/Core/Private/Mac/MacPlatformMisc.cpp
@@ -1931,7 +1931,11 @@ void FMacCrashContext::GenerateCrashInfoAndLaunchReporter() const
 		if(ForkPID == 0)
 		{
 			// Child
-			if(GMacAppInfo.bIsUnattended)
+			if (bImplicitSend)
+			{
+				execl(GMacAppInfo.CrashReportClient, "CrashReportClient", CrashInfoFolder, "-Unattended", "-ImplicitSend", NULL);
+			}
+			else if(GMacAppInfo.bIsUnattended)
 			{
 				execl(GMacAppInfo.CrashReportClient, "CrashReportClient", CrashInfoFolder, "-Unattended", NULL);
 			}
diff --git a/Engine/Source/Runtime/Core/Private/Unix/UnixPlatformCrashContext.cpp b/Engine/Source/Runtime/Core/Private/Unix/UnixPlatformCrashContext.cpp
--- a/Engine/Source/Runtime/Core/Private/Unix/UnixPlatformCrashContext.cpp
+++ b/Engine/Source/Runtime/Core/Private/Unix/UnixPlatformCrashContext.cpp
@@ -565,7 +565,11 @@ void FUnixCrashContext::GenerateCrashInfoAndLaunchReporter(bool bReportingNonCra
 			CrashReportClientArguments += TEXT("\"\"") + CrashReportLogFilepath + TEXT("\"\"");
 			CrashReportClientArguments += TEXT(" ");
 
-			if (bUnattended || bImplicitSend)
+			if (bImplicitSend)
+			{
+				CrashReportClientArguments += TEXT(" -Unattended -ImplicitSend ");
+			}
+			else if (bUnattended)
 			{
 				CrashReportClientArguments += TEXT(" -Unattended ");
 			}
@@ -612,17 +616,6 @@ void FUnixCrashContext::GenerateCrashInfoAndLaunchReporter(bool bReportingNonCra
 				// spin here until CrashReporter exits
 				FProcHandle RunningProc = FPlatformProcess::CreateProc(RelativePathToCrashReporter, *CrashReportClientArguments, true, false, false, NULL, 0, NULL, NULL);
 
-				// We will not have any GUI for the crash reporter if we are sending implicitly, so pop a message box up at least
-				if (bImplicitSend)
-				{
-					if (FApp::CanEverRender())
-					{
-						FPlatformMisc::MessageBoxExt(EAppMsgType::Ok,
-							*NSLOCTEXT("MessageDialog", "ReportCrash_Body", "The application has crashed and will now close. We apologize for the inconvenience.").ToString(),
-							*NSLOCTEXT("MessageDialog", "ReportCrash_Title", "Application Crash Detected").ToString());
-					}
-				}
-
 				// do not wait indefinitely - can be more generous about the hitch than in ensure() case
 				// NOTE: Chris.Wood - increased from 3 to 8 mins because server crashes were timing out and getting lost
 				// NOTE: Do not increase above 8.5 mins without altering watchdog scripts to match
diff --git a/Engine/Source/Runtime/Core/Private/Windows/WindowsPlatformCrashContext.cpp b/Engine/Source/Runtime/Core/Private/Windows/WindowsPlatformCrashContext.cpp
--- a/Engine/Source/Runtime/Core/Private/Windows/WindowsPlatformCrashContext.cpp
+++ b/Engine/Source/Runtime/Core/Private/Windows/WindowsPlatformCrashContext.cpp
@@ -436,7 +436,11 @@ int32 ReportCrashUsingCrashReportClient(FWindowsPlatformCrashContext& InContext,
 			// Pass nullrhi to CRC when the engine is in this mode to stop the CRC attempting to initialize RHI when the capability isn't available
 			bool bNullRHI = !FApp::CanEverRender();
 
-			if (bNoDialog || bNullRHI || bImplicitSend)
+			if (bImplicitSend)
+			{
+				CrashReportClientArguments += TEXT(" -Unattended -ImplicitSend");
+			}
+			else if (bNoDialog || bNullRHI)
 			{
 				CrashReportClientArguments += TEXT(" -Unattended");
 			}
@@ -480,17 +484,6 @@ int32 ReportCrashUsingCrashReportClient(FWindowsPlatformCrashContext& InContext,
 			FString CrashClientPath = FPaths::Combine(*FPaths::EngineDir(), TEXT("Binaries"), FPlatformProcess::GetBinariesSubdirectory(), CrashReportClientExeName);
 			bCrashReporterRan = FPlatformProcess::CreateProc(*CrashClientPath, *CrashReportClientArguments, true, false, false, NULL, 0, NULL, NULL).IsValid();
 
-			if (bImplicitSend)
-			{
-				// We will not have any GUI for the crash reporter if we are sending implicitly, so pop a message box up at least
-				if (FApp::CanEverRender())
-				{
-					FPlatformMisc::MessageBoxExt(EAppMsgType::Ok,
-						*NSLOCTEXT("MessageDialog", "ReportCrash_Body", "The application has crashed and will now close. We apologize for the inconvenience.").ToString(),
-						*NSLOCTEXT("MessageDialog", "ReportCrash_Title", "Application Crash Detected").ToString());
-				}
-			}
-
 			// Restore the dll directory
 			if (CurrentDllDirectory)
 			{
