CrashReportClient: Grey out the option to opt-out of being contacted for internal engine builds.

diff --git a/Engine/Source/Programs/CrashReportClient/Private/CrashDescription.cpp b/Engine/Source/Programs/CrashReportClient/Private/CrashDescription.cpp
--- a/Engine/Source/Programs/CrashReportClient/Private/CrashDescription.cpp
+++ b/Engine/Source/Programs/CrashReportClient/Private/CrashDescription.cpp
@@ -138,7 +138,7 @@ void FPrimaryCrashProperties::Shutdown()
 
 void FPrimaryCrashProperties::UpdateIDs()
 {
-	const bool bAddPersonalData = FCrashReportClientConfig::Get().GetAllowToBeContacted() || FEngineBuildSettings::IsInternalBuild();
+	const bool bAddPersonalData = FCrashReportClientConfig::Get().GetAllowToBeContacted();
 	bAllowToBeContacted = bAddPersonalData;
 	if (bAddPersonalData)
 	{
diff --git a/Engine/Source/Programs/CrashReportClient/Private/CrashReportClientConfig.cpp b/Engine/Source/Programs/CrashReportClient/Private/CrashReportClientConfig.cpp
--- a/Engine/Source/Programs/CrashReportClient/Private/CrashReportClientConfig.cpp
+++ b/Engine/Source/Programs/CrashReportClient/Private/CrashReportClientConfig.cpp
@@ -4,6 +4,7 @@
 #include "Logging/LogMacros.h"
 #include "Misc/ConfigCacheIni.h"
 #include "Misc/App.h"
+#include "Misc/EngineBuildSettings.h"
 #include "CrashReportClientApp.h"
 #include "GenericPlatform/GenericPlatformCrashContext.h"
 
@@ -52,7 +53,11 @@ FCrashReportClientConfig::FCrashReportClientConfig()
 		UE_LOG(CrashReportClientLog, Log, TEXT("DataRouterUrl: %s"), *DataRouterUrl);
 	}
 
-	if (!GConfig->GetBool( TEXT( "CrashReportClient" ), TEXT( "bAllowToBeContacted" ), bAllowToBeContacted, GEngineIni ))
+	if (FEngineBuildSettings::IsInternalBuild())
+	{
+		bAllowToBeContacted = true;
+	}
+	else if (!GConfig->GetBool( TEXT( "CrashReportClient" ), TEXT( "bAllowToBeContacted" ), bAllowToBeContacted, GEngineIni ))
 	{
 		// Default to true when unattended when config is missing. This is mostly for dedicated servers that do not have config files for CRC.
 		if (bUnattended)
diff --git a/Engine/Source/Programs/CrashReportClient/Private/SCrashReportClient.cpp b/Engine/Source/Programs/CrashReportClient/Private/SCrashReportClient.cpp
--- a/Engine/Source/Programs/CrashReportClient/Private/SCrashReportClient.cpp
+++ b/Engine/Source/Programs/CrashReportClient/Private/SCrashReportClient.cpp
@@ -20,6 +20,7 @@
 #include "Widgets/Input/SButton.h"
 #include "Widgets/Layout/SSpacer.h"
 #include "Framework/Application/SlateApplication.h"
+#include "Misc/EngineBuildSettings.h"
 
 #define LOCTEXT_NAMESPACE "CrashReportClient"
 
@@ -236,6 +237,7 @@ void SCrashReportClient::Construct(const FArguments& InArgs, TSharedRef<FCrashRe
 				[
 					SNew(SCheckBox)
 					.IsChecked( FCrashReportClientConfig::Get().GetAllowToBeContacted() ? ECheckBoxState::Checked : ECheckBoxState::Unchecked )
+					.IsEnabled( !FEngineBuildSettings::IsInternalBuild() )
 					.OnCheckStateChanged(CrashReportClient.ToSharedRef(), &FCrashReportClient::AllowToBeContacted_OnCheckStateChanged)
 				]
 
@@ -246,6 +248,7 @@ void SCrashReportClient::Construct(const FArguments& InArgs, TSharedRef<FCrashRe
 				[
 					SNew(STextBlock)
 					.AutoWrapText(true)
+					.IsEnabled( !FEngineBuildSettings::IsInternalBuild() )
 					.Text(LOCTEXT("IAgree", "I agree to be contacted by Epic Games via email if additional information about this crash would help fix it."))
 				]
 			]
