Always send the Pre/Post analytic payload when crashing if enabled

Fix an issue where FCrashReportClient would not call its dtor due to being held by SCrashReportClient

diff --git a/Engine/Source/Programs/CrashReportClient/Private/SCrashReportClient.cpp b/Engine/Source/Programs/CrashReportClient/Private/SCrashReportClient.cpp
--- a/Engine/Source/Programs/CrashReportClient/Private/SCrashReportClient.cpp
+++ b/Engine/Source/Programs/CrashReportClient/Private/SCrashReportClient.cpp
@@ -24,7 +24,7 @@
 
 #define LOCTEXT_NAMESPACE "CrashReportClient"
 
-static void OnBrowserLinkClicked(const FSlateHyperlinkRun::FMetadata& Metadata, TSharedRef<SWidget> ParentWidget)
+static void OnBrowserLinkClicked(const FSlateHyperlinkRun::FMetadata& Metadata)
 {
 	const FString* UrlPtr = Metadata.Find(TEXT("href"));
 	if(UrlPtr)
@@ -33,7 +33,7 @@ static void OnBrowserLinkClicked(const FSlateHyperlinkRun::FMetadata& Metadata,
 	}
 }
 
-static void OnViewCrashDirectory( const FSlateHyperlinkRun::FMetadata& Metadata, TSharedRef<SWidget> ParentWidget )
+static void OnViewCrashDirectory( const FSlateHyperlinkRun::FMetadata& Metadata)
 {
 	const FString* UrlPtr = Metadata.Find( TEXT( "href" ) );
 	if (UrlPtr)
@@ -42,7 +42,7 @@ static void OnViewCrashDirectory( const FSlateHyperlinkRun::FMetadata& Metadata,
 	}
 }
 
-void SCrashReportClient::Construct(const FArguments& InArgs, TSharedRef<FCrashReportClient> Client)
+void SCrashReportClient::Construct(const FArguments& InArgs, const TSharedRef<FCrashReportClient>& Client)
 {
 	CrashReportClient = Client;
 	bHasUserCommentErrors = false;
@@ -106,7 +106,7 @@ void SCrashReportClient::Construct(const FArguments& InArgs, TSharedRef<FCrashRe
 				.Text(CrashDetailedMessage)
 				.AutoWrapText(true)
 				.DecoratorStyleSet( &FCoreStyle::Get() )
-				+ SRichTextBlock::HyperlinkDecorator( TEXT("browser"), FSlateHyperlinkRun::FOnClick::CreateStatic( &OnBrowserLinkClicked, AsShared() ) )
+				+ SRichTextBlock::HyperlinkDecorator( TEXT("browser"), FSlateHyperlinkRun::FOnClick::CreateStatic( &OnBrowserLinkClicked ) )
 			]
 
 			+SVerticalBox::Slot()
@@ -165,7 +165,7 @@ void SCrashReportClient::Construct(const FArguments& InArgs, TSharedRef<FCrashRe
 							.Text( CrashReportDataText )
 							.AutoWrapText( true )
 							.DecoratorStyleSet( &FCrashReportClientStyle::Get() )
-							+ SRichTextBlock::HyperlinkDecorator( TEXT( "browser" ), FSlateHyperlinkRun::FOnClick::CreateStatic( &OnViewCrashDirectory, AsShared() ) )
+							+ SRichTextBlock::HyperlinkDecorator( TEXT( "browser" ), FSlateHyperlinkRun::FOnClick::CreateStatic( &OnViewCrashDirectory ) )
 						]
 					]
 
diff --git a/Engine/Source/Programs/CrashReportClient/Private/SCrashReportClient.h b/Engine/Source/Programs/CrashReportClient/Private/SCrashReportClient.h
--- a/Engine/Source/Programs/CrashReportClient/Private/SCrashReportClient.h
+++ b/Engine/Source/Programs/CrashReportClient/Private/SCrashReportClient.h
@@ -30,7 +30,7 @@ public:
 	 * @param InArgs Slate arguments, not used
 	 * @param Client Crash report client implementation object
 	 */
-	void Construct(const FArguments& InArgs, TSharedRef<FCrashReportClient> Client);
+	void Construct(const FArguments& InArgs, const TSharedRef<FCrashReportClient>& Client);
 
 private:
 	/**
