Only include disaster recovery session data in crash reports when asked to upload log files

#rnx

diff --git a/Engine/Source/Programs/CrashReportClient/Private/CrashReportClientApp.cpp b/Engine/Source/Programs/CrashReportClient/Private/CrashReportClientApp.cpp
--- a/Engine/Source/Programs/CrashReportClient/Private/CrashReportClientApp.cpp
+++ b/Engine/Source/Programs/CrashReportClient/Private/CrashReportClientApp.cpp
@@ -474,7 +474,7 @@ FPlatformErrorReport CollectErrorReport(FRecoveryService* RecoveryService, uint3
 	CrashContext.SerializeAsXML(*CrashContextXMLPath);
 
 #if CRASH_REPORT_WITH_RECOVERY
-	if (RecoveryService && DirectoryExists && SharedCrashContext.bSendUsageData && SharedCrashContext.CrashType != ECrashContextType::Ensure)
+	if (RecoveryService && DirectoryExists && SharedCrashContext.bSendUsageData && SharedCrashContext.CrashType != ECrashContextType::Ensure && FCrashReportCoreConfig::Get().GetSendLogFile())
 	{
 		RecoveryService->CollectFiles(ReportDirectoryAbsolutePath);
 	}
