When the Editor has disabled sending anayltics extend that to the crash report client

diff --git a/Engine/Source/Programs/CrashReportClient/Private/CrashReportClientApp.cpp b/Engine/Source/Programs/CrashReportClient/Private/CrashReportClientApp.cpp
--- a/Engine/Source/Programs/CrashReportClient/Private/CrashReportClientApp.cpp
+++ b/Engine/Source/Programs/CrashReportClient/Private/CrashReportClientApp.cpp
@@ -54,6 +54,9 @@ static FString GameNameFromCmd;
 /** GUID of the crash passed via the command line. */
 static FString CrashGUIDFromCmd;
 
+/** If we want to enable analytics */
+static bool AnalyticsEnabledFromCmd = true;
+
 /**
  * Look for the report to upload, either in the command line or in the platform's report queue
  */
@@ -107,6 +110,11 @@ void ParseCommandLine(const TCHAR* CommandLine)
 		{
 			CrashGUIDFromCmd = Params.FindRef(TEXT("CrashGUID"));
 		}
+
+		if (Switches.Contains(TEXT("NoAnalytics")))
+		{
+			AnalyticsEnabledFromCmd = false;
+		}
 	}
 
 	if (FoundReportDirectoryAbsolutePaths.Num() == 0)
@@ -346,7 +354,10 @@ void RunCrashReportClient(const TCHAR* CommandLine)
 	{
 		ErrorReport.SetCrashReportClientVersion(FCrashReportCoreConfig::Get().GetVersion());
 
-		FCrashReportAnalytics::Initialize();
+		if (AnalyticsEnabledFromCmd)
+		{
+			FCrashReportAnalytics::Initialize();
+		}
 
 		if (bUnattended)
 		{
@@ -368,8 +379,11 @@ void RunCrashReportClient(const TCHAR* CommandLine)
 		}
 #endif // !CRASH_REPORT_UNATTENDED_ONLY
 
-		// Shutdown analytics.
-		FCrashReportAnalytics::Shutdown();
+		if (AnalyticsEnabledFromCmd)
+		{
+			// Shutdown analytics.
+			FCrashReportAnalytics::Shutdown();
+		}
 	}
 	else
 	{
diff --git a/Engine/Source/Runtime/Core/Private/Mac/MacPlatformMisc.cpp b/Engine/Source/Runtime/Core/Private/Mac/MacPlatformMisc.cpp
--- a/Engine/Source/Runtime/Core/Private/Mac/MacPlatformMisc.cpp
+++ b/Engine/Source/Runtime/Core/Private/Mac/MacPlatformMisc.cpp
@@ -1931,7 +1931,15 @@ void FMacCrashContext::GenerateCrashInfoAndLaunchReporter() const
 			}
 			else
 			{
-				execl(GMacAppInfo.CrashReportClient, "CrashReportClient", CrashInfoFolder, NULL);
+				if (bSendUnattendedBugReports)
+				{
+					execl(GMacAppInfo.CrashReportClient, "CrashReportClient", CrashInfoFolder, NULL);
+				}
+				// If the editor setting has been disabled to not send analytics extend this to the CRC
+				else
+				{
+					execl(GMacAppInfo.CrashReportClient, "CrashReportClient", CrashInfoFolder, "-NoAnalytics", NULL);
+				}
 			}
 		}
 		// We no longer wait here because on return the OS will scribble & crash again due to the behaviour of the XPC function
@@ -2010,6 +2018,12 @@ void FMacCrashContext::GenerateEnsureInfoAndLaunchReporter() const
 			Arguments = FString::Printf(TEXT("\"%s/\" -Unattended"), *EnsureLogFolder);
 		}
 
+		// If the editor setting has been disabled to not send analytics extend this to the CRC
+		if (!bSendUnattendedBugReports)
+		{
+			Arguments += TEXT(" -NoAnalytics");
+		}
+
 		FString ReportClient = FPaths::ConvertRelativePathToFull(FPlatformProcess::GenerateApplicationPath(TEXT("CrashReportClient"), EBuildConfigurations::Development));
 		FPlatformProcess::ExecProcess(*ReportClient, *Arguments, nullptr, nullptr, nullptr);
 	}
diff --git a/Engine/Source/Runtime/Core/Private/Unix/UnixPlatformCrashContext.cpp b/Engine/Source/Runtime/Core/Private/Unix/UnixPlatformCrashContext.cpp
--- a/Engine/Source/Runtime/Core/Private/Unix/UnixPlatformCrashContext.cpp
+++ b/Engine/Source/Runtime/Core/Private/Unix/UnixPlatformCrashContext.cpp
@@ -558,6 +558,12 @@ void FUnixCrashContext::GenerateCrashInfoAndLaunchReporter(bool bReportingNonCra
 			CrashReportClientArguments += TEXT("\"\"") + CrashReportLogFilepath + TEXT("\"\"");
 			CrashReportClientArguments += TEXT(" ");
 
+			// If the editor setting has been disabled to not send analytics extend this to the CRC
+			if (!bSendUnattendedBugReports)
+			{
+				CrashReportClientArguments += TEXT(" -NoAnalytics ");
+			}
+
 			if (bUnattended)
 			{
 				CrashReportClientArguments += TEXT(" -Unattended ");
diff --git a/Engine/Source/Runtime/Core/Private/Windows/WindowsPlatformCrashContext.cpp b/Engine/Source/Runtime/Core/Private/Windows/WindowsPlatformCrashContext.cpp
--- a/Engine/Source/Runtime/Core/Private/Windows/WindowsPlatformCrashContext.cpp
+++ b/Engine/Source/Runtime/Core/Private/Windows/WindowsPlatformCrashContext.cpp
@@ -427,6 +427,12 @@ int32 ReportCrashUsingCrashReportClient(FWindowsPlatformCrashContext& InContext,
 			// Run Crash Report Client
 			FString CrashReportClientArguments = FString::Printf(TEXT("\"%s\""), *CrashFolderAbsolute);
 
+			// If the editor setting has been disabled to not send analytics extend this to the CRC
+			if (!bSendUnattendedBugReports)
+			{
+				CrashReportClientArguments += TEXT(" -NoAnalytics ");
+			}
+
 			// Pass nullrhi to CRC when the engine is in this mode to stop the CRC attempting to initialize RHI when the capability isn't available
 			bool bNullRHI = !FApp::CanEverRender();
 
