CIS fix

#rnx

diff --git a/Engine/Source/Programs/CrashReportClient/Private/CrashReportClientApp.cpp b/Engine/Source/Programs/CrashReportClient/Private/CrashReportClientApp.cpp
--- a/Engine/Source/Programs/CrashReportClient/Private/CrashReportClientApp.cpp
+++ b/Engine/Source/Programs/CrashReportClient/Private/CrashReportClientApp.cpp
@@ -612,10 +612,10 @@ void RunCrashReportClient(const TCHAR* CommandLine)
 		double LastTime = FPlatformTime::Seconds();
 		const float IdealFrameTime = 1.0f / IdealFramerate;
 
-		TUniquePtr<FRecoveryService> RecoveryServicePtr;
+		TSharedPtr<FRecoveryService> RecoveryServicePtr; // Note: Shared rather than Unique due to FRecoveryService only being a forward declaration in some builds
 #if CRASH_REPORT_WITH_RECOVERY
 		// Starts the disaster recovery service. This records transactions and allows users to recover from previous crashes.
-		RecoveryServicePtr = MakeUnique<FRecoveryService>(MonitorPid);
+		RecoveryServicePtr = MakeShared<FRecoveryService>(MonitorPid);
 #endif
 
 		FCrashReportAnalytics::Initialize();
