Fix missing icon in CrashReportClient.

#rnx

diff --git a/Engine/Source/Programs/CrashReportClient/Private/Windows/CrashReportClientMainWindows.cpp b/Engine/Source/Programs/CrashReportClient/Private/Windows/CrashReportClientMainWindows.cpp
--- a/Engine/Source/Programs/CrashReportClient/Private/Windows/CrashReportClientMainWindows.cpp
+++ b/Engine/Source/Programs/CrashReportClient/Private/Windows/CrashReportClientMainWindows.cpp
@@ -8,6 +8,8 @@
  */
 int WINAPI WinMain(_In_ HINSTANCE hInInstance, _In_opt_ HINSTANCE hPrevInstance, _In_ LPSTR, _In_ int nCmdShow)
 {
+	hInstance = hInInstance;
+
 	// Run the app
 	RunCrashReportClient(GetCommandLineW());
 
