Fix assert in CrashReportClient when shutting down without sending a crash report.

diff --git a/Engine/Source/Programs/CrashReportClient/Private/CrashReportClientApp.cpp b/Engine/Source/Programs/CrashReportClient/Private/CrashReportClientApp.cpp
--- a/Engine/Source/Programs/CrashReportClient/Private/CrashReportClientApp.cpp
+++ b/Engine/Source/Programs/CrashReportClient/Private/CrashReportClientApp.cpp
@@ -787,6 +787,9 @@ void RunCrashReportClient(const TCHAR* CommandLine)
 			return MonitoredProcess.IsValid() && FPlatformProcess::IsProcRunning(MonitoredProcess) && !FPlatformProcess::GetProcReturnCode(MonitoredProcess, &OutReturnCode);
 		};
 
+		// Keep track of whether we've initialized the analytics system
+		bool bInitializedAnalytics = false;
+
 		// This IsApplicationAlive() call is quite expensive, perform it at low frequency.
 		int32 ApplicationReturnCode = 0;
 		bool bApplicationAlive = IsMonitoredProcessAlive(ApplicationReturnCode);
@@ -805,6 +808,7 @@ void RunCrashReportClient(const TCHAR* CommandLine)
 					if (bReportCrashAnalyticInfo)
 					{
 						FCrashReportAnalytics::Initialize();
+						bInitializedAnalytics = true;
 					}
 
 					// Build error report in memory.
@@ -905,7 +909,11 @@ void RunCrashReportClient(const TCHAR* CommandLine)
 
 		FPlatformProcess::CloseProc(MonitoredProcess);
 
-		FCrashReportAnalytics::Shutdown();
+		if (bInitializedAnalytics)
+		{
+			FCrashReportAnalytics::Shutdown();
+			bInitializedAnalytics = false;
+		}
 	}
 
 	FPrimaryCrashProperties::Shutdown();
