Allow game/editor to continue execution earlier after a crash.

A previous change moved the signal to the game/editor that it's okay to continue to after the crash report client was completely done with sending and resolving callstacks, because it was assumed that there was a syncronization problem. However that proved to be another issue, so moving the signal back to where it was originally. This should make the editor only "freeze" a short time, while necessary data is collected.

(ushell-p4-cherrypick of 9868282 by [REDACTED_PERSON])

diff --git a/Engine/Source/Programs/CrashReportClient/Private/CrashReportClientApp.cpp b/Engine/Source/Programs/CrashReportClient/Private/CrashReportClientApp.cpp
--- a/Engine/Source/Programs/CrashReportClient/Private/CrashReportClientApp.cpp
+++ b/Engine/Source/Programs/CrashReportClient/Private/CrashReportClientApp.cpp
@@ -467,8 +467,8 @@ FPlatformErrorReport CollectErrorReport(FRecoveryService* RecoveryService, uint3
 	// At this point the game can continue execution. It is important this happens
 	// as soon as thread state and minidump has been created, so that ensures cause
 	// as little hitch as possible.
-	//uint8 ResponseCode[] = { 0xd, 0xe, 0xa, 0xd };
-	//FPlatformProcess::WritePipe(WritePipe, ResponseCode, sizeof(ResponseCode));
+	uint8 ResponseCode[] = { 0xd, 0xe, 0xa, 0xd };
+	FPlatformProcess::WritePipe(WritePipe, ResponseCode, sizeof(ResponseCode));
 
 	// Write out the XML file.
 	const FString CrashContextXMLPath = FPaths::Combine(*ReportDirectoryAbsolutePath, FPlatformCrashContext::CrashContextRuntimeXMLNameW);
@@ -825,12 +825,6 @@ void RunCrashReportClient(const TCHAR* CommandLine)
 					const bool bNoDialog = CrashContext.UserSettings.bNoDialog && CrashContext.UserSettings.bSendUnattendedBugReports;
 					const SubmitCrashReportResult Result = SendErrorReport(ErrorReport, bNoDialog);
 
-					// At this point the game can continue execution. It is important this happens
-					// as soon as thread state and minidump has been created, so that ensures cause
-					// as little hitch as possible.
-					uint8 ResponseCode[] = { 0xd, 0xe, 0xa, 0xd };
-					FPlatformProcess::WritePipe(MonitorWritePipe, ResponseCode, sizeof(ResponseCode));
-
 					if (bReportCrashAnalyticInfo)
 					{
 						// If analytics is enabled make sure they are submitted now.
