Prevent analytics being initialized twice

When the user has allowed usage data to be sent we initialize the analytics backend in the crash report client. If the user has also enabled sending unattended reports and an ensure is encountered followed by a crash the crash reporter would assert because the analytics backend was being initialized twice.

diff --git a/Engine/Source/Programs/CrashReportClient/Private/CrashReportClientApp.cpp b/Engine/Source/Programs/CrashReportClient/Private/CrashReportClientApp.cpp
--- a/Engine/Source/Programs/CrashReportClient/Private/CrashReportClientApp.cpp
+++ b/Engine/Source/Programs/CrashReportClient/Private/CrashReportClientApp.cpp
@@ -787,9 +787,6 @@ void RunCrashReportClient(const TCHAR* CommandLine)
 			return MonitoredProcess.IsValid() && FPlatformProcess::IsProcRunning(MonitoredProcess) && !FPlatformProcess::GetProcReturnCode(MonitoredProcess, &OutReturnCode);
 		};
 
-		// Keep track of whether we've initialized the analytics system
-		bool bInitializedAnalytics = false;
-
 		// This IsApplicationAlive() call is quite expensive, perform it at low frequency.
 		int32 ApplicationReturnCode = 0;
 		bool bApplicationAlive = IsMonitoredProcessAlive(ApplicationReturnCode);
@@ -805,10 +802,9 @@ void RunCrashReportClient(const TCHAR* CommandLine)
 				if (IsCrashReportAvailable(MonitorPid, CrashContext, MonitorReadPipe))
 				{
 					const bool bReportCrashAnalyticInfo = CrashContext.UserSettings.bSendUsageData;
-					if (bReportCrashAnalyticInfo)
+					if (bReportCrashAnalyticInfo && !FCrashReportAnalytics::IsAvailable())
 					{
 						FCrashReportAnalytics::Initialize();
-						bInitializedAnalytics = true;
 					}
 
 					// Build error report in memory.
@@ -903,10 +899,9 @@ void RunCrashReportClient(const TCHAR* CommandLine)
 
 		FPlatformProcess::CloseProc(MonitoredProcess);
 
-		if (bInitializedAnalytics)
+		if (FCrashReportAnalytics::IsAvailable())
 		{
 			FCrashReportAnalytics::Shutdown();
-			bInitializedAnalytics = false;
 		}
 	}
 
