Initialize and Shutdown analytics for each crash reported by the monitor

This ensures that we honor the user-settings for reporting analytics correctly if they change while the editor is running.

#rnx

diff --git a/Engine/Source/Programs/CrashReportClient/Private/CrashReportClientApp.cpp b/Engine/Source/Programs/CrashReportClient/Private/CrashReportClientApp.cpp
--- a/Engine/Source/Programs/CrashReportClient/Private/CrashReportClientApp.cpp
+++ b/Engine/Source/Programs/CrashReportClient/Private/CrashReportClientApp.cpp
@@ -802,7 +802,7 @@ void RunCrashReportClient(const TCHAR* CommandLine)
 				if (IsCrashReportAvailable(MonitorPid, CrashContext, MonitorReadPipe))
 				{
 					const bool bReportCrashAnalyticInfo = CrashContext.UserSettings.bSendUsageData;
-					if (bReportCrashAnalyticInfo && !FCrashReportAnalytics::IsAvailable())
+					if (bReportCrashAnalyticInfo)
 					{
 						FCrashReportAnalytics::Initialize();
 					}
@@ -825,6 +825,7 @@ void RunCrashReportClient(const TCHAR* CommandLine)
 					{
 						// If analytics is enabled make sure they are submitted now.
 						FCrashReportAnalytics::GetProvider().BlockUntilFlushed(5.0f);
+						FCrashReportAnalytics::Shutdown();
 					}
 				}
 			}
@@ -869,26 +870,28 @@ void RunCrashReportClient(const TCHAR* CommandLine)
 					bApplicationAlive = IsMonitoredProcessAlive(ApplicationReturnCode);
 					if (!bApplicationAlive)
 					{
-						if (!FCrashReportAnalytics::IsAvailable())
-						{
-							FCrashReportAnalytics::Initialize();
-						}
-
-						// initialize analytics
-						TUniquePtr<FEditorSessionSummarySender> EditorSessionSummarySender = MakeUnique<FEditorSessionSummarySender>(FCrashReportAnalytics::GetProvider(), TEXT("CrashReportClient"), MonitorPid);
-						EditorSessionSummarySender->SetCurrentSessionExitCode(MonitorPid, ApplicationReturnCode);
+						FCrashReportAnalytics::Initialize();
 
-						if (TempCrashContext.UserSettings.bSendUnattendedBugReports)
+						if (FCrashReportAnalytics::IsAvailable())
 						{
-							// Send a spoofed crash report in the case that we detect an abnormal shutdown has occurred
-							if (WasAbnormalShutdown(*EditorSessionSummarySender.Get()))
+							// initialize analytics
+							TUniquePtr<FEditorSessionSummarySender> EditorSessionSummarySender = MakeUnique<FEditorSessionSummarySender>(FCrashReportAnalytics::GetProvider(), TEXT("CrashReportClient"), MonitorPid);
+							EditorSessionSummarySender->SetCurrentSessionExitCode(MonitorPid, ApplicationReturnCode);
+
+							if (TempCrashContext.UserSettings.bSendUnattendedBugReports)
 							{
-								HandleAbnormalShutdown(TempCrashContext, MonitorPid, MonitorWritePipe, RecoveryServicePtr);
+								// Send a spoofed crash report in the case that we detect an abnormal shutdown has occurred
+								if (WasAbnormalShutdown(*EditorSessionSummarySender.Get()))
+								{
+									HandleAbnormalShutdown(TempCrashContext, MonitorPid, MonitorWritePipe, RecoveryServicePtr);
+								}
 							}
+
+							// send analytics and shutdown
+							EditorSessionSummarySender->Shutdown();
 						}
 
-						// send analytics and shutdown
-						EditorSessionSummarySender->Shutdown();
+						FCrashReportAnalytics::Shutdown();
 					}
 				}
 			}
@@ -898,11 +901,6 @@ void RunCrashReportClient(const TCHAR* CommandLine)
 		DeleteTempCrashContextFile(MonitorPid);
 
 		FPlatformProcess::CloseProc(MonitoredProcess);
-
-		if (FCrashReportAnalytics::IsAvailable())
-		{
-			FCrashReportAnalytics::Shutdown();
-		}
 	}
 
 	FPrimaryCrashProperties::Shutdown();
