Add "implicit send" functionality to crash report client monitor path.

Add functionality for bImplicit send configuration variable. This allow a game to automatically send crash report without user interaction, displaying a native os dialog when completed.

diff --git a/Engine/Source/Programs/CrashReportClient/Private/CrashReportClientApp.cpp b/Engine/Source/Programs/CrashReportClient/Private/CrashReportClientApp.cpp
--- a/Engine/Source/Programs/CrashReportClient/Private/CrashReportClientApp.cpp
+++ b/Engine/Source/Programs/CrashReportClient/Private/CrashReportClientApp.cpp
@@ -367,7 +367,7 @@ class FMessageBoxThread : public FRunnable
 	}
 };
 
-SubmitCrashReportResult RunUnattended(FPlatformErrorReport ErrorReport)
+SubmitCrashReportResult RunUnattended(FPlatformErrorReport ErrorReport, bool bImplicitSend)
 {
 	// Set up the main ticker
 	FMainLoopTiming MainLoop(IdealTickRate, EMainLoopOptions::CoreTickerOnly);
@@ -379,7 +379,7 @@ SubmitCrashReportResult RunUnattended(FPlatformErrorReport ErrorReport)
 	FMessageBoxThread MessageBox;
 	FRunnableThread* MessageBoxThread = nullptr;
 
-	if (bImplicitSendFromCmd)
+	if (bImplicitSend)
 	{
 		MessageBoxThread = FRunnableThread::Create(&MessageBox, TEXT("CrashReporter_MessageBox"));
 	}
@@ -390,7 +390,7 @@ SubmitCrashReportResult RunUnattended(FPlatformErrorReport ErrorReport)
 		MainLoop.Tick();
 	}
 
-	if (bImplicitSendFromCmd && MessageBoxThread)
+	if (bImplicitSend && MessageBoxThread)
 	{
 		MessageBoxThread->WaitForCompletion();
 	}
@@ -509,22 +509,25 @@ FPlatformErrorReport CollectErrorReport(FRecoveryService* RecoveryService, uint3
 #endif
 }
 
-SubmitCrashReportResult SendErrorReport(FPlatformErrorReport& ErrorReport, TOptional<bool> bNoDialog = TOptional<bool>())
+SubmitCrashReportResult SendErrorReport(FPlatformErrorReport& ErrorReport, 
+	TOptional<bool> bNoDialogOpt = TOptional<bool>(), 
+	TOptional<bool> bImplicitSendOpt = TOptional<bool>())
 {
 	if (!IsEngineExitRequested() && ErrorReport.HasFilesToUpload() && FPrimaryCrashProperties::Get() != nullptr)
 	{
+		const bool bImplicitSend = bImplicitSendOpt.Get(false);
 		const bool bUnattended =
 #if CRASH_REPORT_UNATTENDED_ONLY
 			true;
 #else
-			bNoDialog.IsSet() ? bNoDialog.GetValue() : FApp::IsUnattended();
+			bNoDialogOpt.Get(FApp::IsUnattended()) || bImplicitSend;
 #endif // CRASH_REPORT_UNATTENDED_ONLY
 
 		ErrorReport.SetCrashReportClientVersion(FCrashReportCoreConfig::Get().GetVersion());
 
 		if (bUnattended)
 		{
-			return RunUnattended(ErrorReport);
+			return RunUnattended(ErrorReport, bImplicitSend);
 		}
 #if !CRASH_REPORT_UNATTENDED_ONLY
 		else
@@ -537,7 +540,7 @@ SubmitCrashReportResult SendErrorReport(FPlatformErrorReport& ErrorReport, TOpti
 				GConfig->GetBool(TEXT("CrashReportClient"), TEXT("CanSendWhenUIFailedToInitialize"), bCanSendWhenUIFailedToInitialize, GEngineIni);
 				if (bCanSendWhenUIFailedToInitialize && !FCrashReportCoreConfig::Get().IsAllowedToCloseWithoutSending())
 				{
-					return RunUnattended(ErrorReport);
+					return RunUnattended(ErrorReport, bImplicitSend);
 				}
 			}
 			return Result;
@@ -759,7 +762,7 @@ void RunCrashReportClient(const TCHAR* CommandLine)
 
 		// Load error report generated by the process from disk
 		FPlatformErrorReport ErrorReport = LoadErrorReport();
-		const SubmitCrashReportResult Result = SendErrorReport(ErrorReport);
+		const SubmitCrashReportResult Result = SendErrorReport(ErrorReport, FApp::IsUnattended(), bImplicitSendFromCmd);
 		// We are not interested in the result of this
 
 		if (AnalyticsEnabledFromCmd)
@@ -822,8 +825,8 @@ void RunCrashReportClient(const TCHAR* CommandLine)
 					}
 #endif
 
-					const bool bNoDialog = CrashContext.UserSettings.bNoDialog && CrashContext.UserSettings.bSendUnattendedBugReports;
-					const SubmitCrashReportResult Result = SendErrorReport(ErrorReport, bNoDialog);
+					const bool bNoDialog = (CrashContext.UserSettings.bNoDialog || CrashContext.UserSettings.bImplicitSend) && CrashContext.UserSettings.bSendUnattendedBugReports;
+					const SubmitCrashReportResult Result = SendErrorReport(ErrorReport, bNoDialog, CrashContext.UserSettings.bImplicitSend);
 
 					if (bReportCrashAnalyticInfo)
 					{
diff --git a/Engine/Source/Runtime/Core/Private/Windows/WindowsPlatformCrashContext.cpp b/Engine/Source/Runtime/Core/Private/Windows/WindowsPlatformCrashContext.cpp
--- a/Engine/Source/Runtime/Core/Private/Windows/WindowsPlatformCrashContext.cpp
+++ b/Engine/Source/Runtime/Core/Private/Windows/WindowsPlatformCrashContext.cpp
@@ -487,6 +487,9 @@ int32 ReportCrashForMonitor(
 	bool bSendUnattendedBugReports = true;
 	bool bSendUsageData = true;
 	bool bCanSendCrashReport = true;
+	// Some projects set this value in non editor builds to automatically send error reports unattended, but display
+	// a plain message box in the crash report client. See CRC app code for details.
+	bool bImplicitSend = false;
 
 	if (FCommandLine::IsInitialized())
 	{
@@ -497,6 +500,14 @@ int32 ReportCrashForMonitor(
 	{
 		GConfig->GetBool(TEXT("/Script/UnrealEd.CrashReportsPrivacySettings"), TEXT("bSendUnattendedBugReports"), bSendUnattendedBugReports, GEditorSettingsIni);
 		GConfig->GetBool(TEXT("/Script/UnrealEd.AnalyticsPrivacySettings"), TEXT("bSendUsageData"), bSendUsageData, GEditorSettingsIni);
+		
+#if !UE_EDITOR
+		if (ReportUI != EErrorReportUI::ReportInUnattendedMode)
+		{
+			// Only check if we are in a non-editor build
+			GConfig->GetBool(TEXT("CrashReportClient"), TEXT("bImplicitSend"), bImplicitSend, GEngineIni);
+		}
+#endif
 	}
 
 #if !UE_EDITOR
@@ -523,6 +534,7 @@ int32 ReportCrashForMonitor(
 	SharedContext->UserSettings.bNoDialog = bNoDialog;
 	SharedContext->UserSettings.bSendUnattendedBugReports = bSendUnattendedBugReports;
 	SharedContext->UserSettings.bSendUsageData = bSendUsageData;
+	SharedContext->UserSettings.bImplicitSend = bImplicitSend;
 
 	SharedContext->SessionContext.bIsExitRequested = IsEngineExitRequested();
 	FCString::Strcpy(SharedContext->ErrorMessage, CR_MAX_ERROR_MESSAGE_CHARS-1, ErrorMessage);
diff --git a/Engine/Source/Runtime/Core/Public/GenericPlatform/GenericPlatformCrashContext.h b/Engine/Source/Runtime/Core/Public/GenericPlatform/GenericPlatformCrashContext.h
--- a/Engine/Source/Runtime/Core/Public/GenericPlatform/GenericPlatformCrashContext.h
+++ b/Engine/Source/Runtime/Core/Public/GenericPlatform/GenericPlatformCrashContext.h
@@ -173,6 +173,7 @@ struct FUserSettingsContext
 	bool					bNoDialog = false;
 	bool					bSendUnattendedBugReports = false;
 	bool					bSendUsageData = false;
+	bool					bImplicitSend = false;
 	TCHAR					LogFilePath[CR_MAX_DIRECTORY_CHARS];
 };
 
