Remove superflous dereference of TCHAR*.

diff --git a/Engine/Source/Programs/CrashReportClient/Private/CrashReportClientApp.cpp b/Engine/Source/Programs/CrashReportClient/Private/CrashReportClientApp.cpp
--- a/Engine/Source/Programs/CrashReportClient/Private/CrashReportClientApp.cpp
+++ b/Engine/Source/Programs/CrashReportClient/Private/CrashReportClientApp.cpp
@@ -459,7 +459,7 @@ FPlatformErrorReport CollectErrorReport(uint32 Pid, const FSharedCrashContext& S
 
 
 	FString ConfigFilename;
-	if (ErrorReport.FindFirstReportFileWithExtension(ConfigFilename, *FGenericCrashContext::CrashConfigExtension))
+	if (ErrorReport.FindFirstReportFileWithExtension(ConfigFilename, FGenericCrashContext::CrashConfigExtension))
 	{
 		FConfigFile CrashConfigFile;
 		CrashConfigFile.Read(ReportDirectoryAbsolutePath / ConfigFilename);
