Fixes to crash report client.

* Flush Crash report client analytics events when running in monitor mode. Previous behaviour was to send analaytics event whenever the crash report client exited. When running  in monitor mode we will need to explicitly flush analytics events after every submission of crash/ensure event.
* Crash report client was not respecting "no dialog" flag set by the runtime, causing dialogs to appear for ensures, which should not display them.
* Allow multiple sends from crash report client. Since the CrashReportClient and CrashReportCoreUnattended classes both use GIsRequesting exit, reset this once a crash report has been sent, to allow consecutive reports to be sent.
* Trying to restart the game after crash would result in multiple instances. The client now captures the result of the user selection and exits accordingly.
* On certain crash types (e.x. general protection fault) the crash reporter would not display callstacks. This was due to a race condition between crash report client doing symbol resolving and when the game was trying print callstack to the log.

diff --git a/Engine/Source/Programs/CrashReportClient/Private/CrashReportClient.cpp b/Engine/Source/Programs/CrashReportClient/Private/CrashReportClient.cpp
--- a/Engine/Source/Programs/CrashReportClient/Private/CrashReportClient.cpp
+++ b/Engine/Source/Programs/CrashReportClient/Private/CrashReportClient.cpp
@@ -37,6 +37,7 @@ FCrashReportClient::FCrashReportClient(const FPlatformErrorReport& InErrorReport
 	, DataRouterUploader(FCrashReportCoreConfig::Get().GetDataRouterURL())
 	, bShouldWindowBeHidden(false)
 	, bSendData(false)
+	, bIsSuccesfullRestart(false)
 {
 	if (FPrimaryCrashProperties::Get()->IsValid())
 	{
@@ -138,6 +139,7 @@ FReply FCrashReportClient::SubmitAndRestart()
 				if (LauncherPlatform->OpenLauncher(OpenOptions))
 				{
 					bLauncherRestarted = true;
+					bIsSuccesfullRestart = true;
 				}
 			}
 		}
@@ -148,6 +150,7 @@ FReply FCrashReportClient::SubmitAndRestart()
 		// Launcher didn't restart the process so start it ourselves
 		const FString CommandLineArguments = FPrimaryCrashProperties::Get()->RestartCommandLine;
 		FPlatformProcess::CreateProc(*CrashedAppPath, *CommandLineArguments, true, false, false, NULL, 0, NULL, NULL);
+		bIsSuccesfullRestart = true;
 	}
 
 	return FReply::Handled();
diff --git a/Engine/Source/Programs/CrashReportClient/Private/CrashReportClient.h b/Engine/Source/Programs/CrashReportClient/Private/CrashReportClient.h
--- a/Engine/Source/Programs/CrashReportClient/Private/CrashReportClient.h
+++ b/Engine/Source/Programs/CrashReportClient/Private/CrashReportClient.h
@@ -127,6 +127,9 @@ public:
 	/** Whether the throbber should be visible while processing the callstack. */
 	EVisibility IsThrobberVisible() const;
 
+	/** Returns true if user has elected to restart crashing process, and if it was a success */
+	bool GetIsSuccesfullRestart() const { return bIsSuccesfullRestart; }
+
 	void AllowToBeContacted_OnCheckStateChanged( ECheckBoxState NewRadioState );
 	void SendLogFile_OnCheckStateChanged( ECheckBoxState NewRadioState );
 
@@ -183,6 +186,9 @@ private:
 	/** Whether we send the data. */
 	bool bSendData;
 
+	/** Store if user has elected to restart crashing process, and if it was a success.*/
+	bool bIsSuccesfullRestart;
+
 };
 
 #endif // !CRASH_REPORT_UNATTENDED_ONLY
diff --git a/Engine/Source/Programs/CrashReportClient/Private/CrashReportClientApp.cpp b/Engine/Source/Programs/CrashReportClient/Private/CrashReportClientApp.cpp
--- a/Engine/Source/Programs/CrashReportClient/Private/CrashReportClientApp.cpp
+++ b/Engine/Source/Programs/CrashReportClient/Private/CrashReportClientApp.cpp
@@ -16,6 +16,7 @@
 #include "Modules/ModuleManager.h"
 #include "HAL/PlatformApplicationMisc.h"
 #include "HAL/PlatformCrashContext.h"
+#include "IAnalyticsProviderET.h"
 
 #if !CRASH_REPORT_UNATTENDED_ONLY
 	#include "SCrashReportClient.h"
@@ -69,6 +70,14 @@ static void* MonitorReadPipe = nullptr;
 /** If in monitor mode, pipe to write data to game. */
 static void* MonitorWritePipe = nullptr;
 
+/** Result of submission of report */
+enum SubmitCrashReportResult {
+	Failed,				// Failed to send report
+	SuccessClosed,		// Succeeded sending report, user has not elected to relaunch
+	SuccessRestarted,	// Succeeded sending report, user has elected to restart process
+	SuccessContinue		// Succeeded sending report, continue running (if monitor mode).
+};
+
 /**
  * Look for the report to upload, either in the command line or in the platform's report queue
  */
@@ -215,7 +224,7 @@ static void OnRequestExit()
 }
 
 #if !CRASH_REPORT_UNATTENDED_ONLY
-bool RunWithUI(FPlatformErrorReport ErrorReport)
+SubmitCrashReportResult RunWithUI(FPlatformErrorReport ErrorReport)
 {
 	// create the platform slate application (what FSlateApplication::Get() returns)
 	TSharedRef<FSlateApplication> Slate = FSlateApplication::Create(MakeShareable(FPlatformApplicationMisc::CreateApplication()));
@@ -247,7 +256,7 @@ bool RunWithUI(FPlatformErrorReport ErrorReport)
 	{
 		// Close down the Slate application
 		FSlateApplication::Shutdown();
-		return false;
+		return Failed;
 	}
 	else if (bRendererFailedToInitializeAtLeastOnce)
 	{
@@ -318,7 +327,8 @@ bool RunWithUI(FPlatformErrorReport ErrorReport)
 	// Close down the Slate application
 	FSlateApplication::Shutdown();
 
-	return true;
+	// Detect if ensure, if user has selected to restart or close.	
+	return CrashReportClient->GetIsSuccesfullRestart() ? SuccessRestarted : (FPrimaryCrashProperties::Get()->bIsEnsure ? SuccessContinue : SuccessClosed);
 }
 #endif // !CRASH_REPORT_UNATTENDED_ONLY
 
@@ -339,7 +349,7 @@ class FMessageBoxThread : public FRunnable
 	}
 };
 
-void RunUnattended(FPlatformErrorReport ErrorReport)
+SubmitCrashReportResult RunUnattended(FPlatformErrorReport ErrorReport)
 {
 	// Set up the main ticker
 	FMainLoopTiming MainLoop(IdealTickRate, EMainLoopOptions::CoreTickerOnly);
@@ -366,6 +376,9 @@ void RunUnattended(FPlatformErrorReport ErrorReport)
 	{
 		MessageBoxThread->WaitForCompletion();
 	}
+
+	// Continue running in case of ensures, otherwise close
+	return FPrimaryCrashProperties::Get()->bIsEnsure ? SuccessContinue : SuccessClosed;
 }
 
 FPlatformErrorReport CollectErrorReport(uint32 Pid, const FSharedCrashContext& SharedCrashContext, void* WritePipe) 
@@ -441,8 +454,8 @@ FPlatformErrorReport CollectErrorReport(uint32 Pid, const FSharedCrashContext& S
 	// At this point the game can continue execution. It is important this happens
 	// as soon as thread state and minidump has been created, so that ensures cause
 	// as little hitch as possible.
-	uint8 ResponseCode[] = { 0xd, 0xe, 0xa, 0xd };
-	FPlatformProcess::WritePipe(WritePipe, ResponseCode, sizeof(ResponseCode));
+	//uint8 ResponseCode[] = { 0xd, 0xe, 0xa, 0xd };
+	//FPlatformProcess::WritePipe(WritePipe, ResponseCode, sizeof(ResponseCode));
 
 	// Write out the XML file.
 	const FString CrashContextXMLPath = FPaths::Combine(*ReportDirectoryAbsolutePath, FPlatformCrashContext::CrashContextRuntimeXMLNameW);
@@ -470,7 +483,7 @@ FPlatformErrorReport CollectErrorReport(uint32 Pid, const FSharedCrashContext& S
 #endif
 }
 
-void SendErrorReport(FPlatformErrorReport& ErrorReport)
+SubmitCrashReportResult SendErrorReport(FPlatformErrorReport& ErrorReport, TOptional<bool> bNoDialog = TOptional<bool>())
 {
 	if (!GIsRequestingExit && ErrorReport.HasFilesToUpload() && FPrimaryCrashProperties::Get() != nullptr)
 	{
@@ -478,32 +491,35 @@ void SendErrorReport(FPlatformErrorReport& ErrorReport)
 #if CRASH_REPORT_UNATTENDED_ONLY
 			true;
 #else
-			FApp::IsUnattended();
+			bNoDialog.IsSet() ? bNoDialog.GetValue() : FApp::IsUnattended();
 #endif // CRASH_REPORT_UNATTENDED_ONLY
 
 		ErrorReport.SetCrashReportClientVersion(FCrashReportCoreConfig::Get().GetVersion());
 
 		if (bUnattended)
 		{
-			RunUnattended(ErrorReport);
+			return RunUnattended(ErrorReport);
 		}
 #if !CRASH_REPORT_UNATTENDED_ONLY
 		else
 		{
-			if (!RunWithUI(ErrorReport))
+			const SubmitCrashReportResult Result = RunWithUI(ErrorReport);
+			if (Result == Failed)
 			{
 				// UI failed to initialize, probably due to driver crash. Send in unattended mode if allowed.
 				bool bCanSendWhenUIFailedToInitialize = true;
 				GConfig->GetBool(TEXT("CrashReportClient"), TEXT("CanSendWhenUIFailedToInitialize"), bCanSendWhenUIFailedToInitialize, GEngineIni);
 				if (bCanSendWhenUIFailedToInitialize && !FCrashReportCoreConfig::Get().IsAllowedToCloseWithoutSending())
 				{
-					RunUnattended(ErrorReport);
+					return RunUnattended(ErrorReport);
 				}
 			}
+			return Result;
 		}
 #endif // !CRASH_REPORT_UNATTENDED_ONLY
 
 	}
+	return Failed;
 }
 
 bool WaitForCrash(uint32 WatchedProcess, FSharedCrashContext& CrashContext, void* ReadPipe)
@@ -548,35 +564,69 @@ void RunCrashReportClient(const TCHAR* CommandLine)
 	ParseCommandLine(CommandLine);
 
 	FPlatformErrorReport::Init();
-	if (AnalyticsEnabledFromCmd)
-	{
-		FCrashReportAnalytics::Initialize();
-	}
 
 	if (MonitorPid && MonitorWritePipe && MonitorReadPipe) 
 	{
-		while (FPlatformProcess::IsApplicationAlive(MonitorPid))
+		while (FPlatformProcess::IsApplicationAlive(MonitorPid) && !GIsRequestingExit)
 		{
 			// Wait for parent process to signal crash
 			FSharedCrashContext CrashContext;
 			if (WaitForCrash(MonitorPid, CrashContext, MonitorReadPipe))
 			{
+				// When running in monitor mode we cannot setup analytics info on the command line (we start earlier than 
+				// config system), instead the crash context contains what the user has selected.
+				if (CrashContext.bSendUsageData)
+				{
+					FCrashReportAnalytics::Initialize();
+				}
+
 				// Build error report in memory
 				FPlatformErrorReport ErrorReport = CollectErrorReport(MonitorPid, CrashContext, MonitorWritePipe);
-				SendErrorReport(ErrorReport);
+				const SubmitCrashReportResult Result = SendErrorReport(ErrorReport, CrashContext.bNoDialog && CrashContext.bSendUnattenededBugReports);
+
+				// At this point the game can continue execution. It is important this happens
+				// as soon as thread state and minidump has been created, so that ensures cause
+				// as little hitch as possible.
+				uint8 ResponseCode[] = { 0xd, 0xe, 0xa, 0xd };
+				FPlatformProcess::WritePipe(MonitorWritePipe, ResponseCode, sizeof(ResponseCode));
+
+				if (AnalyticsEnabledFromCmd)
+				{
+					// If analytics is enabled make sure they are submitted now.
+					FCrashReportAnalytics::GetProvider().BlockUntilFlushed(5.0f);
+				}
+
+				// This is ugly, but since many parts of CrashReportClient and CrashReportCoreUnattended is built
+				// using this flag to exit tick loops, this was the easiest way of making the monitor path be able to
+				// send multiple reports (e.g. ensure followed by crash).
+				if (Result == SuccessContinue)
+				{
+					GIsRequestingExit = false;
+				}
+
+				if (CrashContext.bSendUsageData)
+				{
+					FCrashReportAnalytics::Shutdown();
+				}
 			}
 		}
 	}
 	else
 	{
+		if (AnalyticsEnabledFromCmd)
+		{
+			FCrashReportAnalytics::Initialize();
+		}
+
 		// Load error report generated by the process from disk
 		FPlatformErrorReport ErrorReport = LoadErrorReport();
-		SendErrorReport(ErrorReport);
-	}
+		const SubmitCrashReportResult Result = SendErrorReport(ErrorReport);
+		// We are not interested in the result of this
 
-	if (AnalyticsEnabledFromCmd)
-	{
-		FCrashReportAnalytics::Shutdown();
+		if (AnalyticsEnabledFromCmd)
+		{
+			FCrashReportAnalytics::Shutdown();
+		}
 	}
 
 	FPrimaryCrashProperties::Shutdown();
