Added recovery service to CrashReportClient.
  - The crash reporter client now depends on Concert
  - When started in monitor mode, the main loop ticks.
  - In case of crash, if the disaster recovery plugin was enabled in the client, the disaster recovery can be attached to the bug reports (currently off)

diff --git a/Engine/Plugins/Developer/Concert/ConcertMain/ConcertMain.uplugin b/Engine/Plugins/Developer/Concert/ConcertMain/ConcertMain.uplugin
--- a/Engine/Plugins/Developer/Concert/ConcertMain/ConcertMain.uplugin
+++ b/Engine/Plugins/Developer/Concert/ConcertMain/ConcertMain.uplugin
@@ -15,19 +15,19 @@
 	"Hidden": true,
 	"IsBetaVersion": true,
 	"Installed": false,
-	"SupportedPrograms": [ "UnrealMultiUserServer", "UnrealDisasterRecoveryService" ],
+	"SupportedPrograms": [ "UnrealMultiUserServer", "UnrealDisasterRecoveryService", "CrashReportClient" ],
 	"Modules": [
 		{
 			"Name": "Concert",
 			"Type": "Developer",
 			"LoadingPhase": "PreDefault",
-			"WhitelistPrograms": [ "UnrealMultiUserServer", "UnrealDisasterRecoveryService" ]
+			"WhitelistPrograms": [ "UnrealMultiUserServer", "UnrealDisasterRecoveryService", "CrashReportClient" ]
 		},
 		{
 			"Name": "ConcertTransport",
 			"Type": "Developer",
 			"LoadingPhase": "PreDefault",
-			"WhitelistPrograms": [ "UnrealMultiUserServer", "UnrealDisasterRecoveryService" ]
+			"WhitelistPrograms": [ "UnrealMultiUserServer", "UnrealDisasterRecoveryService", "CrashReportClient" ]
 		}
 	]
 }
diff --git a/Engine/Plugins/Developer/Concert/ConcertMain/Source/Concert/Private/ConcertServer.cpp b/Engine/Plugins/Developer/Concert/ConcertMain/Source/Concert/Private/ConcertServer.cpp
--- a/Engine/Plugins/Developer/Concert/ConcertMain/Source/Concert/Private/ConcertServer.cpp
+++ b/Engine/Plugins/Developer/Concert/ConcertMain/Source/Concert/Private/ConcertServer.cpp
@@ -37,8 +37,8 @@ FString GetArchiveName(const FString& SessionName, const FConcertSessionSettings
 FConcertServerPaths::FConcertServerPaths(const FString& InRole, const FString& InWorkingDir, const FString& InSavedDir)
 	: WorkingDir(InWorkingDir.Len() ? InWorkingDir / InRole : (FPaths::ProjectIntermediateDir() / TEXT("Concert") / InRole))
 	, SavedDir(InSavedDir.Len() ? InSavedDir / InRole : (FPaths::ProjectSavedDir() / TEXT("Concert") / InRole))
-	, BaseWorkingDir(InWorkingDir)
-	, BaseSavedDir(InSavedDir)
+	, BaseWorkingDir(InWorkingDir.Len() ? InWorkingDir : FPaths::ProjectIntermediateDir())
+	, BaseSavedDir(InSavedDir.Len() ? InSavedDir : FPaths::ProjectSavedDir())
 {
 }
 
@@ -362,7 +362,7 @@ void FConcertServer::RecoverSessions()
 			// Remove the oldest sessions
 			for (const FSavedSessionInfo& SortedSession : SortedSessions)
 			{
-				ConcertUtil::DeleteDirectoryTree(*Paths->GetSessionWorkingDir(ArchivedSessionInfos[SortedSession.Key].SessionId), *Paths->GetBaseWorkingDir());
+				ConcertUtil::DeleteDirectoryTree(*Paths->GetSessionSavedDir(ArchivedSessionInfos[SortedSession.Key].SessionId), *Paths->GetBaseSavedDir());	
 			}
 
 			// Update the list of sessions to restore
@@ -424,6 +424,11 @@ FGuid FConcertServer::ArchiveSession(const FGuid& SessionId, const FString& Arch
 	return ArchivedSessionId;
 }
 
+bool FConcertServer::ExportSession(const FGuid& SessionId, const FConcertSessionFilter& SessionFilter, const FString& DestDir, bool bAnonymizeData, FText& OutFailureReason)
+{
+	return EventSink->ExportSession(*this, SessionId, DestDir, SessionFilter, bAnonymizeData);
+}
+
 bool FConcertServer::RenameSession(const FGuid& SessionId, const FString& NewName, FText& OutFailureReason)
 {
 	// NOTE: This function is exposed to the server internals and should not be directly called by connected clients. Clients
diff --git a/Engine/Plugins/Developer/Concert/ConcertMain/Source/Concert/Private/ConcertServer.h b/Engine/Plugins/Developer/Concert/ConcertMain/Source/Concert/Private/ConcertServer.h
--- a/Engine/Plugins/Developer/Concert/ConcertMain/Source/Concert/Private/ConcertServer.h
+++ b/Engine/Plugins/Developer/Concert/ConcertMain/Source/Concert/Private/ConcertServer.h
@@ -101,6 +101,7 @@ public:
 	virtual TSharedPtr<IConcertServerSession> CreateSession(const FConcertSessionInfo& SessionInfo, FText& OutFailureReason) override;
 	virtual TSharedPtr<IConcertServerSession> RestoreSession(const FGuid& SessionId, const FConcertSessionInfo& SessionInfo, const FConcertSessionFilter& SessionFilter, FText& OutFailureReason) override;
 	virtual FGuid ArchiveSession(const FGuid& SessionId, const FString& ArchiveNameOverride, const FConcertSessionFilter& SessionFilter, FText& OutFailureReason) override;
+	virtual bool ExportSession(const FGuid& SessionId, const FConcertSessionFilter& SessionFilter, const FString& DestDir, bool bAnonymizeData, FText& OutFailureReason) override;
 	virtual bool RenameSession(const FGuid& SessionId, const FString& NewName, FText& OutFailureReason) override;
 	virtual bool DestroySession(const FGuid& SessionId, FText& OutFailureReason) override;
 	virtual TArray<FConcertSessionClientInfo> GetSessionClients(const FGuid& SessionId) const override;
diff --git a/Engine/Plugins/Developer/Concert/ConcertMain/Source/Concert/Public/ConcertMessageData.h b/Engine/Plugins/Developer/Concert/ConcertMain/Source/Concert/Public/ConcertMessageData.h
--- a/Engine/Plugins/Developer/Concert/ConcertMain/Source/Concert/Public/ConcertMessageData.h
+++ b/Engine/Plugins/Developer/Concert/ConcertMain/Source/Concert/Public/ConcertMessageData.h
@@ -222,6 +222,10 @@ struct FConcertSessionFilter
 	/** True if only live data should be included (live transactions and head package revisions) */
 	UPROPERTY()
 	bool bOnlyLiveData = false;
+
+	/** True to export the activity summaries without the package/transaction data to look at the log rather than replaying the activities. */
+	UPROPERTY()
+	bool bMetaDataOnly = false;
 };
 
 USTRUCT()
diff --git a/Engine/Plugins/Developer/Concert/ConcertMain/Source/Concert/Public/IConcertServer.h b/Engine/Plugins/Developer/Concert/ConcertMain/Source/Concert/Public/IConcertServer.h
--- a/Engine/Plugins/Developer/Concert/ConcertMain/Source/Concert/Public/IConcertServer.h
+++ b/Engine/Plugins/Developer/Concert/ConcertMain/Source/Concert/Public/IConcertServer.h
@@ -119,6 +119,17 @@ public:
 	 */
 	virtual FGuid ArchiveSession(const FGuid& SessionId, const FString& ArchiveNameOverride, const FConcertSessionFilter& SessionFilter, FText& OutFailureReason) = 0;
 
+	/**
+	 * Copy the session data to a destination folder for external usage.
+	 * @param SessionId The ID of the session to export. Can be a live or archived session.
+	 * @param SessionFilter The filter controlling which activities from the session should be archived.
+	 * @param DestDir The directory where the exported files should be copied. (Must exist)
+	 * @param bAnonymizeData True to obfuscate the object and package names stored in the database. Ignored if FConcertSessionFilter::bMetaDataOnly is false.
+	 * @param OutFailureReason The reason why exporting the session would fail.
+	 * @return True if the session files were exported successfully, false otherwise.
+	 */
+	virtual bool ExportSession(const FGuid& SessionId, const FConcertSessionFilter& SessionFilter, const FString& DestDir, bool bAnonymizeData, FText& OutFailureReason) = 0;
+
 	/**
 	 * Rename a live or archived Concert session on the server. The server automatically detects if the specified session Id is a live or an
 	 * archived session.
diff --git a/Engine/Plugins/Developer/Concert/ConcertMain/Source/Concert/Public/IConcertServerEventSink.h b/Engine/Plugins/Developer/Concert/ConcertMain/Source/Concert/Public/IConcertServerEventSink.h
--- a/Engine/Plugins/Developer/Concert/ConcertMain/Source/Concert/Public/IConcertServerEventSink.h
+++ b/Engine/Plugins/Developer/Concert/ConcertMain/Source/Concert/Public/IConcertServerEventSink.h
@@ -56,6 +56,12 @@ public:
 	 */
 	virtual bool ArchiveSession(const IConcertServer& InServer, const FString& InLiveSessionWorkingDir, const FString& InArchivedSessionRoot, const FConcertSessionInfo& InArchivedSessionInfo, const FConcertSessionFilter& InSessionFilter) = 0;
 
+	/**
+	 * Called to migrate and gather the data of a live or archived session to a directory for external usage.
+	 * @note As opposed to archiving a session, once exported, the session is not tracked anymore by the server and the exported files can be copied/modified while the server is running.
+	 */
+	virtual bool ExportSession(const IConcertServer& InServer, const FGuid& InSessionId, const FString& DestDir, const FConcertSessionFilter& InSessionFilter, bool bAnonymizeData) = 0;
+
 	/**
 	 * Called to migrate the data for an archived session into a live session.
 	 * @note OnLiveSessionCreated will be called if this restoration was successful.
diff --git a/Engine/Plugins/Developer/Concert/ConcertSync/ConcertSyncCore/ConcertSyncCore.uplugin b/Engine/Plugins/Developer/Concert/ConcertSync/ConcertSyncCore/ConcertSyncCore.uplugin
--- a/Engine/Plugins/Developer/Concert/ConcertSync/ConcertSyncCore/ConcertSyncCore.uplugin
+++ b/Engine/Plugins/Developer/Concert/ConcertSync/ConcertSyncCore/ConcertSyncCore.uplugin
@@ -15,13 +15,13 @@
 	"Hidden": true,
 	"IsBetaVersion": true,
 	"Installed": false,
-	"SupportedPrograms": [ "UnrealMultiUserServer", "UnrealDisasterRecoveryService" ],
+	"SupportedPrograms": [ "UnrealMultiUserServer", "UnrealDisasterRecoveryService", "CrashReportClient" ],
 	"Modules": [
 		{
 			"Name": "ConcertSyncCore",
 			"Type": "Developer",
 			"LoadingPhase": "PreDefault",
-			"WhitelistPrograms": [ "UnrealMultiUserServer", "UnrealDisasterRecoveryService" ]
+			"WhitelistPrograms": [ "UnrealMultiUserServer", "UnrealDisasterRecoveryService", "CrashReportClient" ]
 		}
 	],
 	"Plugins": [
diff --git a/Engine/Plugins/Developer/Concert/ConcertSync/ConcertSyncCore/Source/ConcertSyncCore/Private/ConcertSyncSessionDatabase.cpp b/Engine/Plugins/Developer/Concert/ConcertSync/ConcertSyncCore/Source/ConcertSyncCore/Private/ConcertSyncSessionDatabase.cpp
--- a/Engine/Plugins/Developer/Concert/ConcertSync/ConcertSyncCore/Source/ConcertSyncCore/Private/ConcertSyncSessionDatabase.cpp
+++ b/Engine/Plugins/Developer/Concert/ConcertSync/ConcertSyncCore/Source/ConcertSyncCore/Private/ConcertSyncSessionDatabase.cpp
@@ -1276,20 +1276,20 @@ bool FConcertSyncSessionDatabase::SetLockActivity(const FConcertSyncLockActivity
 		);
 }
 
-bool FConcertSyncSessionDatabase::SetTransactionActivity(const FConcertSyncTransactionActivity& InTransactionActivity)
+bool FConcertSyncSessionDatabase::SetTransactionActivity(const FConcertSyncTransactionActivity& InTransactionActivity, const bool bMetaDataOnly)
 {
 	FConcertSyncSessionDatabaseScopedTransaction ScopedTransaction(*Statements);
 	return ScopedTransaction.CommitOrRollback(
-		SetTransactionEvent(InTransactionActivity.EventId, InTransactionActivity.EventData) &&
+		SetTransactionEvent(InTransactionActivity.EventId, InTransactionActivity.EventData, bMetaDataOnly) &&
 		Statements->SetActivityData(InTransactionActivity.ActivityId, InTransactionActivity.EndpointId, InTransactionActivity.EventTime, InTransactionActivity.EventType, InTransactionActivity.EventId, InTransactionActivity.EventSummary)
 		);
 }
 
-bool FConcertSyncSessionDatabase::SetPackageActivity(const FConcertSyncPackageActivity& InPackageActivity)
+bool FConcertSyncSessionDatabase::SetPackageActivity(const FConcertSyncPackageActivity& InPackageActivity, const bool bMetaDataOnly)
 {
 	FConcertSyncSessionDatabaseScopedTransaction ScopedTransaction(*Statements);
 	return ScopedTransaction.CommitOrRollback(
-		SetPackageEvent(InPackageActivity.EventId, InPackageActivity.EventData) &&
+		SetPackageEvent(InPackageActivity.EventId, InPackageActivity.EventData, bMetaDataOnly) &&
 		Statements->SetActivityData(InPackageActivity.ActivityId, InPackageActivity.EndpointId, InPackageActivity.EventTime, InPackageActivity.EventType, InPackageActivity.EventId, InPackageActivity.EventSummary)
 		);
 }
@@ -1613,15 +1613,19 @@ bool FConcertSyncSessionDatabase::AddTransactionEvent(const FConcertSyncTransact
 	return SetTransactionEvent(OutTransactionEventId, InTransactionEvent);
 }
 
-bool FConcertSyncSessionDatabase::SetTransactionEvent(const int64 InTransactionEventId, const FConcertSyncTransactionEvent& InTransactionEvent)
+bool FConcertSyncSessionDatabase::SetTransactionEvent(const int64 InTransactionEventId, const FConcertSyncTransactionEvent& InTransactionEvent, const bool bMetaDataOnly)
 {
 	// Write the data blob file
 	const FString TransactionDataFilename = TransactionDataUtil::GetDataFilename(InTransactionEventId);
 	const FString TransactionDataPathname = TransactionDataUtil::GetDataPath(SessionPath) / TransactionDataFilename;
+
 	FStructOnScope Transaction(FConcertTransactionFinalizedEvent::StaticStruct(), (uint8*)&InTransactionEvent.Transaction);
-	if (!SaveTransaction(TransactionDataPathname, Transaction))
+	if (!bMetaDataOnly)
 	{
-		return false;
+		if (!SaveTransaction(TransactionDataPathname, Transaction))
+		{
+			return false;
+		}
 	}
 
 	// Add the database entry
@@ -1832,12 +1836,12 @@ bool FConcertSyncSessionDatabase::AddPackageEvent(const FConcertSyncPackageEvent
 	return SetPackageEvent(OutPackageEventId, PackageRevision, InPackageEvent.Package);
 }
 
-bool FConcertSyncSessionDatabase::SetPackageEvent(const int64 InPackageEventId, const FConcertSyncPackageEvent& InPackageEvent)
+bool FConcertSyncSessionDatabase::SetPackageEvent(const int64 InPackageEventId, const FConcertSyncPackageEvent& InPackageEvent, const bool bMetaDataOnly)
 {
-	return SetPackageEvent(InPackageEventId, InPackageEvent.PackageRevision, InPackageEvent.Package);
+	return SetPackageEvent(InPackageEventId, InPackageEvent.PackageRevision, InPackageEvent.Package, bMetaDataOnly);
 }
 
-bool FConcertSyncSessionDatabase::SetPackageEvent(const int64 InPackageEventId, const int64 InPackageRevision, const FConcertPackage& InPackage)
+bool FConcertSyncSessionDatabase::SetPackageEvent(const int64 InPackageEventId, const int64 InPackageRevision, const FConcertPackage& InPackage, const bool bMetaDataOnly)
 {
 	if (!ensureAlwaysMsgf(InPackageRevision > 0, TEXT("Invalid package revision! Must be greater than zero.")))
 	{
@@ -1857,10 +1861,14 @@ bool FConcertSyncSessionDatabase::SetPackageEvent(const int64 InPackageEventId,
 
 	// Write the data blob file
 	const FString PackageDataFilename = PackageDataUtil::GetDataFilename(InPackage.Info.PackageName, InPackageRevision);
-	const FString PackageDataPathname = PackageDataUtil::GetDataPath(SessionPath) / PackageDataFilename;
-	if (!SavePackage(PackageDataPathname, InPackage.Info, InPackage.PackageData))
+
+	if (!bMetaDataOnly)
 	{
-		return false;
+		const FString PackageDataPathname = PackageDataUtil::GetDataPath(SessionPath) / PackageDataFilename;
+		if (!SavePackage(PackageDataPathname, InPackage.Info, InPackage.PackageData))
+		{
+			return false;
+		}
 	}
 
 	// Add the database entry
diff --git a/Engine/Plugins/Developer/Concert/ConcertSync/ConcertSyncCore/Source/ConcertSyncCore/Public/ConcertSyncSessionDatabase.h b/Engine/Plugins/Developer/Concert/ConcertSync/ConcertSyncCore/Source/ConcertSyncCore/Public/ConcertSyncSessionDatabase.h
--- a/Engine/Plugins/Developer/Concert/ConcertSync/ConcertSyncCore/Source/ConcertSyncCore/Public/ConcertSyncSessionDatabase.h
+++ b/Engine/Plugins/Developer/Concert/ConcertSync/ConcertSyncCore/Source/ConcertSyncCore/Public/ConcertSyncSessionDatabase.h
@@ -149,10 +149,11 @@ public:
 	 * @note This function is expected to be called on the client to populate its version of the session database from data synced from the server.
 	 *
 	 * @param InTransactionActivity		The transaction activity to set.
+	 * @param bMetaDataOnly				True to store the meta data only, omitting the transaction data required to replay the transaction.
 	 *
 	 * @return True if the transaction activity was set, false otherwise.
 	 */
-	bool SetTransactionActivity(const FConcertSyncTransactionActivity& InTransactionActivity);
+	bool SetTransactionActivity(const FConcertSyncTransactionActivity& InTransactionActivity, const bool bMetaDataOnly = false);
 
 	/**
 	 * Set a package activity in this database, creating or replacing it.
@@ -160,10 +161,11 @@ public:
 	 * @note This function is expected to be called on the client to populate its version of the session database from data synced from the server.
 	 *
 	 * @param InPackageActivity		The package activity to set.
+	 * @param bMetaDataOnly			True to store the meta data only, omitting the package data itself.
 	 *
 	 * @return True if the package activity was set, false otherwise.
 	 */
-	bool SetPackageActivity(const FConcertSyncPackageActivity& InPackageActivity);
+	bool SetPackageActivity(const FConcertSyncPackageActivity& InPackageActivity, const bool bMetaDataOnly = false);
 
 	/**
 	 * Get the generic part of an activity from this database.
@@ -637,10 +639,11 @@ private:
 	 *
 	 * @param InTransactionEventId		The ID of the transaction event to set.
 	 * @param InTransactionEvent		The transaction event to set.
+	 * @param bMetaDataOnly				True to store the meta data only, omitting the transaction data.
 	 *
 	 * @return True if the transaction event was set, false otherwise.
 	 */
-	bool SetTransactionEvent(const int64 InTransactionEventId, const FConcertSyncTransactionEvent& InTransactionEvent);
+	bool SetTransactionEvent(const int64 InTransactionEventId, const FConcertSyncTransactionEvent& InTransactionEvent, const bool bMetaDataOnly = false);
 
 	/**
 	 * Enumerate the IDs of any live transaction events for the given package name ID.
@@ -667,10 +670,11 @@ private:
 	 *
 	 * @param InPackageEventId			The ID of the package event to set.
 	 * @param InPackageEvent			The package event to set.
+	 * @param bMetaDataOnly				True to store the meta data only, omitting the package data.
 	 *
 	 * @return True if the package event was set, false otherwise.
 	 */
-	bool SetPackageEvent(const int64 InPackageEventId, const FConcertSyncPackageEvent& InPackageEvent);
+	bool SetPackageEvent(const int64 InPackageEventId, const FConcertSyncPackageEvent& InPackageEvent, const bool bMetaDataOnly = false);
 
 	/**
 	 * Set a package event in this database, creating or replacing it.
@@ -678,10 +682,11 @@ private:
 	 * @param InPackageEventId			The ID of the package event to set.
 	 * @param InPackageRevision			The package revision to set.
 	 * @param InPackage					The package to set.
+	 * @param bMetaDataOnly				True to store the meta data only, omitting the package data.
 	 *
 	 * @return True if the package event was set, false otherwise.
 	 */
-	bool SetPackageEvent(const int64 InPackageEventId, const int64 InPackageRevision, const FConcertPackage& InPackage);
+	bool SetPackageEvent(const int64 InPackageEventId, const int64 InPackageRevision, const FConcertPackage& InPackage, const bool bMetaDataOnly = false);
 
 	/**
 	 * Get the maximum ID of the package events in this database.
diff --git a/Engine/Plugins/Developer/Concert/ConcertSync/ConcertSyncServer/ConcertSyncServer.uplugin b/Engine/Plugins/Developer/Concert/ConcertSync/ConcertSyncServer/ConcertSyncServer.uplugin
--- a/Engine/Plugins/Developer/Concert/ConcertSync/ConcertSyncServer/ConcertSyncServer.uplugin
+++ b/Engine/Plugins/Developer/Concert/ConcertSync/ConcertSyncServer/ConcertSyncServer.uplugin
@@ -15,13 +15,13 @@
 	"Hidden": true,
 	"IsBetaVersion": true,
 	"Installed": false,
-	"SupportedPrograms": [ "UnrealMultiUserServer", "UnrealDisasterRecoveryService" ],
+	"SupportedPrograms": [ "UnrealMultiUserServer", "UnrealDisasterRecoveryService", "CrashReportClient" ],
 	"Modules": [
 		{
 			"Name": "ConcertSyncServer",
 			"Type": "Developer",
 			"LoadingPhase": "Default",
-			"WhitelistPrograms": [ "UnrealMultiUserServer", "UnrealDisasterRecoveryService" ]
+			"WhitelistPrograms": [ "UnrealMultiUserServer", "UnrealDisasterRecoveryService", "CrashReportClient" ]
 		}
 	],
 	"Plugins": [
diff --git a/Engine/Plugins/Developer/Concert/ConcertSync/ConcertSyncServer/Source/ConcertSyncServer/Private/ConcertSyncServer.cpp b/Engine/Plugins/Developer/Concert/ConcertSync/ConcertSyncServer/Source/ConcertSyncServer/Private/ConcertSyncServer.cpp
--- a/Engine/Plugins/Developer/Concert/ConcertSync/ConcertSyncServer/Source/ConcertSyncServer/Private/ConcertSyncServer.cpp
+++ b/Engine/Plugins/Developer/Concert/ConcertSync/ConcertSyncServer/Source/ConcertSyncServer/Private/ConcertSyncServer.cpp
@@ -157,11 +157,11 @@ bool MigrateSessionData(const FConcertSyncSessionDatabase& InSourceDatabase, con
 				}
 				if (ConcertSyncSessionDatabaseFilterUtil::TransactionEventPassesFilter(TransactionActivity.EventId, InDestSessionFilter, InSourceDatabase))
 				{
-					if (!InSourceDatabase.GetTransactionEvent(TransactionActivity.EventId, TransactionActivity.EventData))
+					if (!InSourceDatabase.GetTransactionEvent(TransactionActivity.EventId, TransactionActivity.EventData, InDestSessionFilter.bMetaDataOnly))
 					{
 						MIGRATE_SET_ERROR_RESULT_AND_RETURN("Failed to get transaction event '%s' from database at '%s': %s", *LexToString(TransactionActivity.EventId), *InSourceDatabase.GetFilename(), *InSourceDatabase.GetLastError());
 					}
-					if (!DestDatabase.SetTransactionActivity(TransactionActivity))
+					if (!DestDatabase.SetTransactionActivity(TransactionActivity, InDestSessionFilter.bMetaDataOnly))
 					{
 						MIGRATE_SET_ERROR_RESULT_AND_RETURN("Failed to set transaction activity '%s' on database at '%s': %s", *LexToString(InActivityId), *DestDatabase.GetFilename(), *DestDatabase.GetLastError());
 					}
@@ -178,11 +178,11 @@ bool MigrateSessionData(const FConcertSyncSessionDatabase& InSourceDatabase, con
 				}
 				if (ConcertSyncSessionDatabaseFilterUtil::PackageEventPassesFilter(PackageActivity.EventId, InDestSessionFilter, InSourceDatabase))
 				{
-					if (!InSourceDatabase.GetPackageEvent(PackageActivity.EventId, PackageActivity.EventData))
+					if (!InSourceDatabase.GetPackageEvent(PackageActivity.EventId, PackageActivity.EventData, InDestSessionFilter.bMetaDataOnly))
 					{
 						MIGRATE_SET_ERROR_RESULT_AND_RETURN("Failed to get package event '%s' from database at '%s': %s", *LexToString(PackageActivity.EventId), *InSourceDatabase.GetFilename(), *InSourceDatabase.GetLastError());
 					}
-					if (!DestDatabase.SetPackageActivity(PackageActivity))
+					if (!DestDatabase.SetPackageActivity(PackageActivity, InDestSessionFilter.bMetaDataOnly))
 					{
 						MIGRATE_SET_ERROR_RESULT_AND_RETURN("Failed to set package activity '%s' on database at '%s': %s", *LexToString(InActivityId), *DestDatabase.GetFilename(), *DestDatabase.GetLastError());
 					}
@@ -308,6 +308,22 @@ bool FConcertSyncServer::ArchiveSession(const IConcertServer& InServer, const FS
 	return RetVal;
 }
 
+bool FConcertSyncServer::ExportSession(const IConcertServer& InServer, const FGuid& InSessionId, const FString& DestDir, const FConcertSessionFilter& InSessionFilter, bool bAnonymizeData)
+{
+	if (TSharedPtr<FConcertSyncServerLiveSession> LiveSession = LiveSessions.FindRef(InSessionId)) // If the session is live.
+	{
+		ConcertSyncServerUtils::WriteSessionInfoToDirectory(DestDir, LiveSession->GetSession().GetSessionInfo());
+		return ConcertSyncServerUtils::MigrateSessionData(LiveSession->GetSessionDatabase(), DestDir, InSessionFilter);
+	}
+	else if (TSharedPtr<FConcertSyncServerArchivedSession> ArchivedSession = ArchivedSessions.FindRef(InSessionId)) // If the session is archived.
+	{
+		ConcertSyncServerUtils::WriteSessionInfoToDirectory(DestDir, ArchivedSession->GetSessionInfo());
+		return ConcertSyncServerUtils::MigrateSessionData(ArchivedSession->GetSessionDatabase(), DestDir, InSessionFilter);
+	}
+
+	return false; // Session not found.
+}
+
 bool FConcertSyncServer::RestoreSession(const IConcertServer& InServer, const FGuid& InArchivedSessionId, const FString& InLiveSessionRoot, const FConcertSessionInfo& InLiveSessionInfo, const FConcertSessionFilter& InSessionFilter)
 {
 	if (TSharedPtr<FConcertSyncServerArchivedSession> ArchivedSession = ArchivedSessions.FindRef(InArchivedSessionId))
diff --git a/Engine/Plugins/Developer/Concert/ConcertSync/ConcertSyncServer/Source/ConcertSyncServer/Private/ConcertSyncServer.h b/Engine/Plugins/Developer/Concert/ConcertSync/ConcertSyncServer/Source/ConcertSyncServer/Private/ConcertSyncServer.h
--- a/Engine/Plugins/Developer/Concert/ConcertSync/ConcertSyncServer/Source/ConcertSyncServer/Private/ConcertSyncServer.h
+++ b/Engine/Plugins/Developer/Concert/ConcertSync/ConcertSyncServer/Source/ConcertSyncServer/Private/ConcertSyncServer.h
@@ -34,6 +34,7 @@ public:
 	virtual void OnArchivedSessionDestroyed(const IConcertServer& InServer, const FGuid& InArchivedSessionId) override;
 	virtual bool ArchiveSession(const IConcertServer& InServer, TSharedRef<IConcertServerSession> InLiveSession, const FString& InArchivedSessionRoot, const FConcertSessionInfo& InArchivedSessionInfo, const FConcertSessionFilter& InSessionFilter) override;
 	virtual bool ArchiveSession(const IConcertServer& InServer, const FString& InLiveSessionWorkingDir, const FString& InArchivedSessionRoot, const FConcertSessionInfo& InArchivedSessionInfo, const FConcertSessionFilter& InSessionFilter) override;
+	virtual bool ExportSession(const IConcertServer& InServer, const FGuid& InSessionId, const FString& DestDir, const FConcertSessionFilter& InSessionFilter, bool bAnonymizeData) override;
 	virtual bool RestoreSession(const IConcertServer& InServer, const FGuid& InArchivedSessionId, const FString& InLiveSessionRoot, const FConcertSessionInfo& InLiveSessionInfo, const FConcertSessionFilter& InSessionFilter) override;
 	virtual bool GetSessionActivities(const IConcertServer& InServer, const FGuid& SessionId, int64 FromActivityId, int64 ActivityCount, TArray<FConcertSessionSerializedPayload>& OutActivities) override;
 	virtual void OnLiveSessionRenamed(const IConcertServer& InServer, TSharedRef<IConcertServerSession> InLiveSession) override;
diff --git a/Engine/Plugins/Developer/Concert/DisasterRecoveryClient/Source/DisasterRecoveryClient/Private/DisasterRecoveryClientModule.cpp b/Engine/Plugins/Developer/Concert/DisasterRecoveryClient/Source/DisasterRecoveryClient/Private/DisasterRecoveryClientModule.cpp
--- a/Engine/Plugins/Developer/Concert/DisasterRecoveryClient/Source/DisasterRecoveryClient/Private/DisasterRecoveryClientModule.cpp
+++ b/Engine/Plugins/Developer/Concert/DisasterRecoveryClient/Source/DisasterRecoveryClient/Private/DisasterRecoveryClientModule.cpp
@@ -72,6 +72,15 @@ public:
 	}
 
 private:
+	constexpr bool ShouldLaunchDisasterRecoveryBackendService() const
+	{
+#if PLATFORM_WINDOWS // On windows, the service is launched by the crash analytics process.
+		return false;
+#else
+		return true;
+#endif
+	}
+
 	void RegisterSettings()
 	{
 		if (ISettingsModule* SettingsModule = FModuleManager::GetModulePtr<ISettingsModule>("Settings"))
@@ -166,11 +175,14 @@ private:
 			}
 		}
 
-		CheckDisasterRecoveryServiceHealthTickHandle = FTicker::GetCoreTicker().AddTicker(TEXT("CheckDisasterRecoveryServiceHealth"), 1.0f, [this](float)
+		if (ShouldLaunchDisasterRecoveryBackendService())
 		{
-			CheckDisasterRecoveryServiceHealth();
-			return true;
-		});
+			CheckDisasterRecoveryServiceHealthTickHandle = FTicker::GetCoreTicker().AddTicker(TEXT("CheckDisasterRecoveryServiceHealth"), 1.0f, [this](float)
+			{
+				CheckDisasterRecoveryServiceHealth();
+				return true;
+			});
+		}
 	}
 
 	/** Returns the folder where the disaster recovery service should keep the live session files (the working directory). */
@@ -271,30 +283,33 @@ private:
 			return false;
 		}
 
-		// Find the service path that will host the sync server
-		const FString DisasterRecoveryServicePath = GetDisasterRecoveryServicePath();
-		if (DisasterRecoveryServicePath.IsEmpty())
-		{
-			UE_LOG(LogDisasterRecovery, Warning, TEXT("Disaster Recovery Service application was not found. Disaster Recovery will be disabled! Please build 'UnrealDisasterRecoveryService'."));
-			return false;
-		}
+		const FString DisasterRecoveryServerName = FString::Printf(TEXT("RecoverySvr_%d"), FPlatformProcess::GetCurrentProcessId());
+		const FString DisasterRecoverySessionName = FString::Printf(TEXT("Recovery_%d_%s_%s"), FPlatformProcess::GetCurrentProcessId(), FApp::GetProjectName(), *FDateTime::Now().ToString());
 
-		const FString DisasterRecoveryServerName = FString::Printf(TEXT("%s_%d"), *FApp::GetInstanceId().ToString(), DisasterRecoveryServiceInstanceCount++);
-		const FString DisasterRecoverySessionName = FString::Printf(TEXT("%s_%s_%s"), FApp::GetProjectName(), *FDateTime::Now().ToString(), *DisasterRecoveryServerName);
+		if (ShouldLaunchDisasterRecoveryBackendService())
+		{
+			// Find the service path that will host the sync server
+			const FString DisasterRecoveryServicePath = GetDisasterRecoveryServicePath();
+			if (DisasterRecoveryServicePath.IsEmpty())
+			{
+				UE_LOG(LogDisasterRecovery, Warning, TEXT("Disaster Recovery Service application was not found. Disaster Recovery will be disabled! Please build 'UnrealDisasterRecoveryService'."));
+				return false;
+			}
 
-		FString DisasterRecoveryServiceCommandLine;
-		DisasterRecoveryServiceCommandLine += FString::Printf(TEXT(" -ConcertServer=\"%s\""), *DisasterRecoveryServerName);
-		DisasterRecoveryServiceCommandLine += TEXT(" -ConcertIgnore");
-		DisasterRecoveryServiceCommandLine += FString::Printf(TEXT(" -EditorPID=%d"), FPlatformProcess::GetCurrentProcessId());
-		DisasterRecoveryServiceCommandLine += FString::Printf(TEXT(" -ConcertWorkingDir=\"%s\""), *GetDefaultServerWorkingDir());
-		DisasterRecoveryServiceCommandLine += FString::Printf(TEXT(" -ConcertSavedDir=\"%s\""), *GetDefaultServerArchiveDir());
+			FString DisasterRecoveryServiceCommandLine;
+			DisasterRecoveryServiceCommandLine += FString::Printf(TEXT(" -ConcertServer=\"%s\""), *DisasterRecoveryServerName);
+			DisasterRecoveryServiceCommandLine += TEXT(" -ConcertIgnore");
+			DisasterRecoveryServiceCommandLine += FString::Printf(TEXT(" -EditorPID=%d"), FPlatformProcess::GetCurrentProcessId());
+			DisasterRecoveryServiceCommandLine += FString::Printf(TEXT(" -ConcertWorkingDir=\"%s\""), *GetDefaultServerWorkingDir());
+			DisasterRecoveryServiceCommandLine += FString::Printf(TEXT(" -ConcertSavedDir=\"%s\""), *GetDefaultServerArchiveDir());
 
-		// Create the service process that will host the sync server
-		DisasterRecoveryServiceHandle = FPlatformProcess::CreateProc(*DisasterRecoveryServicePath, *DisasterRecoveryServiceCommandLine, true, true, true, nullptr, 0, nullptr, nullptr, nullptr);
-		if (!DisasterRecoveryServiceHandle.IsValid())
-		{
-			UE_LOG(LogDisasterRecovery, Error, TEXT("Failed to launch Disaster Recovery Service application. Disaster Recovery will be disabled!"));
-			return false;
+			// Create the service process that will host the sync server
+			DisasterRecoveryServiceHandle = FPlatformProcess::CreateProc(*DisasterRecoveryServicePath, *DisasterRecoveryServiceCommandLine, true, true, true, nullptr, 0, nullptr, nullptr, nullptr);
+			if (!DisasterRecoveryServiceHandle.IsValid())
+			{
+				UE_LOG(LogDisasterRecovery, Error, TEXT("Failed to launch Disaster Recovery Service application. Disaster Recovery will be disabled!"));
+				return false;
+			}
 		}
 
 		bool bRecoveringFromCrash = InSessionInfoToRestore && InSessionInfoToRestore->bAutoRestoreLastSession;
diff --git a/Engine/Plugins/Messaging/UdpMessaging/UdpMessaging.uplugin b/Engine/Plugins/Messaging/UdpMessaging/UdpMessaging.uplugin
--- a/Engine/Plugins/Messaging/UdpMessaging/UdpMessaging.uplugin
+++ b/Engine/Plugins/Messaging/UdpMessaging/UdpMessaging.uplugin
@@ -19,7 +19,8 @@
 		"UnrealFrontend",
 		"UnrealMultiUserServer",
 		"UnrealDisasterRecoveryService",
-		"UnrealLightmass"
+		"UnrealLightmass",
+		"CrashReportClient"
 	],
 	"Modules" :
 	[
@@ -35,8 +36,9 @@
 				"UnrealFrontend",
 				"UnrealMultiUserServer",
 				"UnrealDisasterRecoveryService",
-				"UnrealLightmass"
+				"UnrealLightmass",
+				"CrashReportClient"
 			]
 		}
 	]
-}
\ No newline at end of file
+}
diff --git a/Engine/Programs/CrashReportClient/Config/DefaultEngine.ini b/Engine/Programs/CrashReportClient/Config/DefaultEngine.ini
--- a/Engine/Programs/CrashReportClient/Config/DefaultEngine.ini
+++ b/Engine/Programs/CrashReportClient/Config/DefaultEngine.ini
@@ -9,3 +9,11 @@ UIInitRetryInterval=2.0
 
 [/Script/Engine.NetworkSettings]
 n.VerifyPeer=false
+
+[Plugins]
++ProgramEnabledPlugins="UdpMessaging"
++ProgramEnabledPlugins="ConcertSyncServer"
+
+[/Script/UdpMessaging.UdpMessagingSettings]
+EnableTransport=True
+MulticastTimeToLive=0
diff --git a/Engine/Source/Programs/CrashReportClient/CrashReportClient.Build.cs b/Engine/Source/Programs/CrashReportClient/CrashReportClient.Build.cs
--- a/Engine/Source/Programs/CrashReportClient/CrashReportClient.Build.cs
+++ b/Engine/Source/Programs/CrashReportClient/CrashReportClient.Build.cs
@@ -31,9 +31,11 @@ public class CrashReportClient : ModuleRules
 				"Analytics",
 				"AnalyticsET",
 				"DesktopPlatform",
-                		"LauncherPlatform"
+				"LauncherPlatform",
+				"Concert",
+				"InputCore",
 			}
-			);
+		);
 
 		PrivateDependencyModuleNames.AddRange(
 			new string[] 
@@ -48,12 +50,15 @@ public class CrashReportClient : ModuleRules
 		PrivateIncludePathModuleNames.AddRange(
 			new string[] {
 				"SlateReflector",
+				"ConcertSyncCore",
+				"ConcertSyncServer",
 			}
 		);
 
 		DynamicallyLoadedModuleNames.AddRange(
 			new string[] {
 				"SlateReflector",
+				"ConcertSyncServer",
 			}
 		);
 
diff --git a/Engine/Source/Programs/CrashReportClient/CrashReportClient.Target.cs b/Engine/Source/Programs/CrashReportClient/CrashReportClient.Target.cs
--- a/Engine/Source/Programs/CrashReportClient/CrashReportClient.Target.cs
+++ b/Engine/Source/Programs/CrashReportClient/CrashReportClient.Target.cs
@@ -15,6 +15,11 @@ public class CrashReportClientTarget : TargetRules
 
 		LaunchModuleName = "CrashReportClient";
 
+		AdditionalPlugins.Add("UdpMessaging");
+		AdditionalPlugins.Add("ConcertSyncServer");
+		bBuildWithEditorOnlyData = false;
+		bCompileWithPluginSupport = true; // Enable Developer plugins (like Concert!)
+
 		if (Target.Platform != UnrealTargetPlatform.Linux)
 		{
 			ExtraModuleNames.Add("EditorStyle");
@@ -22,7 +27,7 @@ public class CrashReportClientTarget : TargetRules
 
         bLegalToDistributeBinary = true;
 
-		bBuildDeveloperTools = false;
+		bBuildDeveloperTools = true;
 
 		// CrashReportClient doesn't ever compile with the engine linked in
 		bCompileAgainstEngine = false;
diff --git a/Engine/Source/Programs/CrashReportClient/Private/CrashReportClientApp.cpp b/Engine/Source/Programs/CrashReportClient/Private/CrashReportClientApp.cpp
--- a/Engine/Source/Programs/CrashReportClient/Private/CrashReportClientApp.cpp
+++ b/Engine/Source/Programs/CrashReportClient/Private/CrashReportClientApp.cpp
@@ -34,6 +34,7 @@
 
 #include "PlatformErrorReport.h"
 #include "XmlFile.h"
+#include "RecoveryService.h"
 
 /** Default main window size */
 const FVector2D InitialWindowDimensions(740, 560);
@@ -70,6 +71,9 @@ static void* MonitorReadPipe = nullptr;
 /** If in monitor mode, pipe to write data to game. */
 static void* MonitorWritePipe = nullptr;
 
+/** True to enable the disaster recovery service. */
+static bool bDisasterRecoveryServiceEnabled = false;
+
 /** Result of submission of report */
 enum SubmitCrashReportResult {
 	Failed,				// Failed to send report
@@ -142,6 +146,11 @@ void ParseCommandLine(const TCHAR* CommandLine)
 			AnalyticsEnabledFromCmd = false;
 		}
 
+		if (Params.Contains(TEXT("ConcertServer")))
+		{
+			bDisasterRecoveryServiceEnabled = true;
+		}
+
 		CrashGUIDFromCmd = Params.FindRef(TEXT("CrashGUID"));
 		MonitorPid = FPlatformString::Atoi64(*Params.FindRef(TEXT("MONITOR")));
 		MonitorReadPipe = (void*) FPlatformString::Atoi64(*Params.FindRef(TEXT("READ")));
@@ -381,7 +390,7 @@ SubmitCrashReportResult RunUnattended(FPlatformErrorReport ErrorReport)
 	return FPrimaryCrashProperties::Get()->bIsEnsure ? SuccessContinue : SuccessClosed;
 }
 
-FPlatformErrorReport CollectErrorReport(uint32 Pid, const FSharedCrashContext& SharedCrashContext, void* WritePipe) 
+FPlatformErrorReport CollectErrorReport(FRecoveryService* RecoveryService, uint32 Pid, const FSharedCrashContext& SharedCrashContext, void* WritePipe)
 {
 	// @note: This API is only partially implemented on Mac OS and Linux.
 	FProcHandle ProcessHandle = FPlatformProcess::OpenProcess(Pid);
@@ -461,6 +470,11 @@ FPlatformErrorReport CollectErrorReport(uint32 Pid, const FSharedCrashContext& S
 	const FString CrashContextXMLPath = FPaths::Combine(*ReportDirectoryAbsolutePath, FPlatformCrashContext::CrashContextRuntimeXMLNameW);
 	CrashContext.SerializeAsXML(*CrashContextXMLPath);
 
+	if (RecoveryService && DirectoryExists && SharedCrashContext.bSendUsageData && SharedCrashContext.CrashType != ECrashContextType::Ensure || SharedCrashContext.CrashType == ECrashContextType::Assert)
+	{
+		RecoveryService->CollectFiles(ReportDirectoryAbsolutePath);
+	}
+
 	const TCHAR* CrachContextBuffer = *CrashContext.GetBuffer();
 	FPrimaryCrashProperties::Set(new FCrashContext(ReportDirectoryAbsolutePath / TEXT("CrashContext.runtime-xml"), CrachContextBuffer));
 
@@ -470,7 +484,6 @@ FPlatformErrorReport CollectErrorReport(uint32 Pid, const FSharedCrashContext& S
 	return ErrorReport;
 #else
 
-
 	FString ConfigFilename;
 	if (ErrorReport.FindFirstReportFileWithExtension(ConfigFilename, FGenericCrashContext::CrashConfigExtension))
 	{
@@ -522,21 +535,16 @@ SubmitCrashReportResult SendErrorReport(FPlatformErrorReport& ErrorReport, TOpti
 	return Failed;
 }
 
-bool WaitForCrash(uint32 WatchedProcess, FSharedCrashContext& CrashContext, void* ReadPipe)
+bool IsCrashReportAvailable(uint32 WatchedProcess, FSharedCrashContext& CrashContext, void* ReadPipe)
 {
-	TArray<uint8> Buffer;
+	static TArray<uint8> Buffer;
+	Buffer.Reserve(8 * 1024); // This allocates only once because Buffer is static.
 
-	while (FPlatformProcess::IsApplicationAlive(WatchedProcess))
+	if (FPlatformProcess::ReadPipeToArray(ReadPipe, Buffer))
 	{
-		if (FPlatformProcess::ReadPipeToArray(ReadPipe, Buffer))
-		{
-			FPlatformMemory::Memcpy(&CrashContext, Buffer.GetData(), Buffer.Num());
-			return true;
-		}
-		
-		FPlatformProcess::Sleep(1.0f);
+		FPlatformMemory::Memcpy(&CrashContext, Buffer.GetData(), Buffer.Num());
+		return true;
 	}
-
 	return false;
 }
 
@@ -548,8 +556,11 @@ void RunCrashReportClient(const TCHAR* CommandLine)
 	// Increase the HttpSendTimeout to 5 minutes
 	GConfig->SetFloat(TEXT("HTTP"), TEXT("HttpSendTimeout"), 5 * 60.0f, GEngineIni);
 
-	// Set up the main loop
-	GEngineLoop.PreInit(CommandLine);
+	// Initialize the engine. -Messaging enables MessageBus transports required by Concert (Recovery Service).
+	FString FinalCommandLine(CommandLine);
+	FinalCommandLine += TEXT(" -Messaging");
+	GEngineLoop.PreInit(*FinalCommandLine);
+	check(GConfig && GConfig->IsReadyForUse());
 
 	// Make sure all UObject classes are registered and default properties have been initialized
 	ProcessNewlyLoadedUObjects();
@@ -557,6 +568,12 @@ void RunCrashReportClient(const TCHAR* CommandLine)
 	// Tell the module manager is may now process newly-loaded UObjects when new C++ modules are loaded
 	FModuleManager::Get().StartProcessingNewlyLoadedObjects();
 
+	// Load internal Concert plugins in the pre-default phase
+	IPluginManager::Get().LoadModulesForEnabledPlugins(ELoadingPhase::PreDefault);
+
+	// Load Concert Sync plugins in default phase
+	IPluginManager::Get().LoadModulesForEnabledPlugins(ELoadingPhase::Default);
+
 	// Initialize config.
 	FCrashReportCoreConfig::Get();
 
@@ -565,53 +582,7 @@ void RunCrashReportClient(const TCHAR* CommandLine)
 
 	FPlatformErrorReport::Init();
 
-	if (MonitorPid && MonitorWritePipe && MonitorReadPipe) 
-	{
-		while (FPlatformProcess::IsApplicationAlive(MonitorPid) && !GIsRequestingExit)
-		{
-			// Wait for parent process to signal crash
-			FSharedCrashContext CrashContext;
-			if (WaitForCrash(MonitorPid, CrashContext, MonitorReadPipe))
-			{
-				// When running in monitor mode we cannot setup analytics info on the command line (we start earlier than 
-				// config system), instead the crash context contains what the user has selected.
-				if (CrashContext.bSendUsageData)
-				{
-					FCrashReportAnalytics::Initialize();
-				}
-
-				// Build error report in memory
-				FPlatformErrorReport ErrorReport = CollectErrorReport(MonitorPid, CrashContext, MonitorWritePipe);
-				const SubmitCrashReportResult Result = SendErrorReport(ErrorReport, CrashContext.bNoDialog && CrashContext.bSendUnattenededBugReports);
-
-				// At this point the game can continue execution. It is important this happens
-				// as soon as thread state and minidump has been created, so that ensures cause
-				// as little hitch as possible.
-				uint8 ResponseCode[] = { 0xd, 0xe, 0xa, 0xd };
-				FPlatformProcess::WritePipe(MonitorWritePipe, ResponseCode, sizeof(ResponseCode));
-
-				if (AnalyticsEnabledFromCmd)
-				{
-					// If analytics is enabled make sure they are submitted now.
-					FCrashReportAnalytics::GetProvider().BlockUntilFlushed(5.0f);
-				}
-
-				// This is ugly, but since many parts of CrashReportClient and CrashReportCoreUnattended is built
-				// using this flag to exit tick loops, this was the easiest way of making the monitor path be able to
-				// send multiple reports (e.g. ensure followed by crash).
-				if (Result == SuccessContinue)
-				{
-					GIsRequestingExit = false;
-				}
-
-				if (CrashContext.bSendUsageData)
-				{
-					FCrashReportAnalytics::Shutdown();
-				}
-			}
-		}
-	}
-	else
+	if (MonitorPid == 0) // Does not monitor any process.
 	{
 		if (AnalyticsEnabledFromCmd)
 		{
@@ -628,12 +599,105 @@ void RunCrashReportClient(const TCHAR* CommandLine)
 			FCrashReportAnalytics::Shutdown();
 		}
 	}
+	else // Launched in 'service mode - watches/serves a process'
+	{
+		const int32 IdealFramerate = 30;
+		double LastTime = FPlatformTime::Seconds();
+		const float IdealFrameTime = 1.0f / IdealFramerate;
+
+		TUniquePtr<FRecoveryService> RecoveryService;
+		if (bDisasterRecoveryServiceEnabled)
+		{
+			// Starts the disaster recovery service. This records transactions and allows users to recover from previous crashes.
+			RecoveryService = MakeUnique<FRecoveryService>(MonitorPid);
+		}
+
+		// This IsApplicationAlive() call is quite expensive, perform it at low frequency.
+		bool bApplicationAlive = FPlatformProcess::IsApplicationAlive(MonitorPid);
+		while (bApplicationAlive && !GIsRequestingExit)
+		{
+			const double CurrentTime = FPlatformTime::Seconds();
+
+			// If 'out-of-process' crash reporting was enabled.
+			if (MonitorWritePipe && MonitorReadPipe)
+			{
+				// Check if the monitored process signaled a crash or an ensure.
+				FSharedCrashContext CrashContext;
+				if (IsCrashReportAvailable(MonitorPid, CrashContext, MonitorReadPipe))
+				{
+					// When running in monitor mode we cannot setup analytics info on the command line (we start earlier than 
+					// config system), instead the crash context contains what the user has selected.
+					bool bReportCrashAnalyticInfo = CrashContext.bSendUsageData;
+					if (bReportCrashAnalyticInfo)
+					{
+						FCrashReportAnalytics::Initialize();
+					}
+
+					// Build error report in memory. (NOTE: Don't collect recovery service files, data anonymity is not implemented yet)
+					FPlatformErrorReport ErrorReport = CollectErrorReport(nullptr/*RecoveryService.Get()*/, MonitorPid, CrashContext, MonitorWritePipe);
+					const SubmitCrashReportResult Result = SendErrorReport(ErrorReport, CrashContext.bNoDialog && CrashContext.bSendUnattenededBugReports);
+
+					// At this point the game can continue execution. It is important this happens
+					// as soon as thread state and minidump has been created, so that ensures cause
+					// as little hitch as possible.
+					uint8 ResponseCode[] = { 0xd, 0xe, 0xa, 0xd };
+					FPlatformProcess::WritePipe(MonitorWritePipe, ResponseCode, sizeof(ResponseCode));
+
+					if (bReportCrashAnalyticInfo)
+					{
+						// If analytics is enabled make sure they are submitted now.
+						FCrashReportAnalytics::GetProvider().BlockUntilFlushed(5.0f);
+					}
+
+					// This is ugly, but since many parts of CrashReportClient and CrashReportCoreUnattended is built
+					// using this flag to exit tick loops, this was the easiest way of making the monitor path be able to
+					// send multiple reports (e.g. ensure followed by crash).
+					if (Result == SuccessContinue)
+					{
+						GIsRequestingExit = false;
+					}
+
+					if (bReportCrashAnalyticInfo)
+					{
+						FCrashReportAnalytics::Shutdown();
+					}
+				}
+			}
+
+			FTaskGraphInterface::Get().ProcessThreadUntilIdle(ENamedThreads::GameThread);
+
+			// Pump & Tick objects
+			const double DeltaTime = CurrentTime - LastTime;
+			FTicker::GetCoreTicker().Tick(DeltaTime);
+
+			GFrameCounter++;
+			FStats::AdvanceFrame(false);
+			GLog->FlushThreadedLogs();
+
+			// Run garbage collection for the UObjects for the rest of the frame or at least to 2 ms
+			IncrementalPurgeGarbage(true, FMath::Max<float>(0.002f, IdealFrameTime - (FPlatformTime::Seconds() - LastTime)));
+
+			// Throttle main thread main fps by sleeping if we still have time
+			FPlatformProcess::Sleep(FMath::Max<float>(0.0f, IdealFrameTime - (FPlatformTime::Seconds() - LastTime)));
+
+			// Check if the application is alive about every second. (This is an expensive call)
+			if (GFrameCounter % IdealFramerate == 0)
+			{
+				bApplicationAlive = FPlatformProcess::IsApplicationAlive(MonitorPid);
+			}
+
+			LastTime = CurrentTime;
+		}
+	}
 
 	FPrimaryCrashProperties::Shutdown();
 	FPlatformErrorReport::ShutDown();
 
 	RequestEngineExit(TEXT("CrashReportClientApp RequestExit"));
 
+	// Allow the game thread to finish processing any latent tasks.
+	FTaskGraphInterface::Get().ProcessThreadUntilIdle(ENamedThreads::GameThread);
+
 	FEngineLoop::AppPreExit();
 	FModuleManager::Get().UnloadModulesAtShutdown();
 	FTaskGraphInterface::Shutdown();
diff --git a/Engine/Source/Programs/CrashReportClient/Private/RecoveryService.cpp b/Engine/Source/Programs/CrashReportClient/Private/RecoveryService.cpp
new file mode 100644
--- /dev/null
+++ b/Engine/Source/Programs/CrashReportClient/Private/RecoveryService.cpp
@@ -0,0 +1,122 @@
+// Copyright 1998-2019 Epic Games, Inc. All Rights Reserved.
+
+#include "RecoveryService.h"
+#include "HAL/FileManager.h"
+#include "CrashReportClient.h" // For CrashReportClientLog
+#include "Interfaces/IPluginManager.h"
+#include "ConcertSettings.h"
+#include "ConcertSyncSessionFlags.h"
+#include "IConcertServer.h"
+#include "IConcertSession.h"
+#include "IConcertSyncServer.h"
+#include "IConcertSyncServerModule.h"
+#include "ConcertMessageData.h"
+#include "Runtime/Launch/Resources/Version.h"
+
+static const TCHAR RecoveryServiceName[] = TEXT("Disaster Recovery Service");
+
+bool FRecoveryService::CollectFiles(const FString& DestDir, bool bMetaDataOnly, bool bAnonymizeMetaData)
+{
+	auto LogError = [](const TCHAR* Reason)
+	{
+		UE_LOG(CrashReportClientLog, Error, TEXT("Failed to collect recovery session file(s). %s"), Reason);
+	};
+
+	if (!Server)
+	{
+		LogError(TEXT("The recovery service is not running anymore."));
+		return false;
+	}
+	else if (!IFileManager::Get().DirectoryExists(*DestDir))
+	{
+		LogError(TEXT("The destination folder doesn't exist."));
+		return false;
+	}
+
+	FGuid ExportedSessionId = GetRecoverySessionId();
+	if (!ExportedSessionId.IsValid())
+	{
+		LogError(TEXT("The session session could not be found."));
+		return false;
+	}
+
+	FText ErrorMsg;
+	FConcertSessionFilter Filter;
+	Filter.bMetaDataOnly = bMetaDataOnly;
+	if (!Server->GetConcertServer()->ExportSession(ExportedSessionId, Filter, DestDir, bAnonymizeMetaData, ErrorMsg))
+	{
+		LogError(TEXT("Server failed to export the session."));
+		return false;
+	}
+
+	return true;
+}
+
+bool FRecoveryService::Startup()
+{
+	if (!IConcertSyncServerModule::IsAvailable())
+	{
+		UE_LOG(CrashReportClientLog, Error, TEXT("ConcertSyncServer Module is missing. Recovery service will be disabled!"));
+		return false;
+	}
+
+	TSharedPtr<IPlugin> Plugin = IPluginManager::Get().FindPlugin("UdpMessaging");
+	if (!Plugin || !Plugin->IsEnabled())
+	{
+		// The UdpMessaging plugin should be added to the {appname}.Target.cs build file.
+		UE_LOG(CrashReportClientLog, Error, TEXT("The 'UDP Messaging' plugin is disabled. The Concert server only supports UDP protocol. Recovery service will be disabled!"));
+		return false;
+	}
+
+	// Setup the disaster recovery server configuration
+	UConcertServerConfig* ServerConfig = IConcertSyncServerModule::Get().ParseServerSettings(FCommandLine::Get());
+	ServerConfig->bAutoArchiveOnReboot = true; // If server crashed, was killed, etc, ensure the recovery session is archived (expected by recovery flow).
+	ServerConfig->NumSessionsToKeep = 10;
+	ServerConfig->EndpointSettings.RemoteEndpointTimeoutSeconds = 0;
+	ServerConfig->WorkingDir = GetRecoveryServiceWorkingDir();
+	ServerConfig->ArchiveDir = GetRecoveryServiceArchivedDir();
+
+	// Start disaster recovery server.
+	Server = IConcertSyncServerModule::Get().CreateServer(TEXT("DisasterRecovery"));
+	Server->Startup(ServerConfig, EConcertSyncSessionFlags::Default_DisasterRecoverySession);
+
+	UE_LOG(CrashReportClientLog, Display, TEXT("%s Initialized (Name: %s, Version: %d.%d, Role: %s)"), RecoveryServiceName, *Server->GetConcertServer()->GetServerInfo().ServerName, ENGINE_MAJOR_VERSION, ENGINE_MINOR_VERSION, *Server->GetConcertServer()->GetRole());
+	return true;
+}
+
+void FRecoveryService::Shutdown()
+{
+	if (Server)
+	{
+		Server->Shutdown();
+		Server.Reset();
+		UE_LOG(CrashReportClientLog, Display, TEXT("%s Shutdown"), RecoveryServiceName);
+	}
+}
+
+FString FRecoveryService::GetRecoveryServiceWorkingDir() const
+{
+	return FPaths::ProjectIntermediateDir();
+}
+
+FString FRecoveryService::GetRecoveryServiceArchivedDir() const
+{
+	return FPaths::ProjectSavedDir();
+}
+
+FGuid FRecoveryService::GetRecoverySessionId() const
+{
+	// As convention, the disaster recovery client names the session as Recover_{pid}_{projName}_{dateTime}.
+	FString SearchedSessionName = FString::Printf(TEXT("Recovery_%d"), MonitorPid);
+
+	// As long as the Concert server is up, the session would remain live (it's going to be archived when the server shutdown or reboot).
+	for (TSharedPtr<IConcertServerSession>& Session : Server->GetConcertServer()->GetSessions())
+	{
+		if (Session->GetName().StartsWith(SearchedSessionName))
+		{
+			return Session->GetId();
+		}
+	}
+
+	return FGuid();// Not found.
+}
diff --git a/Engine/Source/Programs/CrashReportClient/Private/RecoveryService.h b/Engine/Source/Programs/CrashReportClient/Private/RecoveryService.h
new file mode 100644
--- /dev/null
+++ b/Engine/Source/Programs/CrashReportClient/Private/RecoveryService.h
@@ -0,0 +1,30 @@
+// Copyright 1998-2019 Epic Games, Inc. All Rights Reserved.
+
+#pragma once
+
+#include "CoreMinimal.h"
+
+class IConcertSyncServer;
+
+/** Manage the life time of the Disaster recovery service. */
+class FRecoveryService
+{
+public:
+	FRecoveryService(uint64 InClientPid) : MonitorPid(InClientPid) { Startup(); }
+	~FRecoveryService() { Shutdown(); }
+
+	/** Collects the session files. */
+	bool CollectFiles(const FString& DestDir, bool bMetaDataOnly = true, bool bAnonymizeMetaData = false);
+
+private:
+	bool Startup();
+	void Shutdown();
+
+	FString GetRecoveryServiceWorkingDir() const;
+	FString GetRecoveryServiceArchivedDir() const;
+	FGuid GetRecoverySessionId() const;
+
+private:
+	TSharedPtr<IConcertSyncServer> Server;
+	uint64 MonitorPid = 0;
+};
diff --git a/Engine/Source/Runtime/Core/Private/Windows/WindowsPlatformCrashContext.cpp b/Engine/Source/Runtime/Core/Private/Windows/WindowsPlatformCrashContext.cpp
--- a/Engine/Source/Runtime/Core/Private/Windows/WindowsPlatformCrashContext.cpp
+++ b/Engine/Source/Runtime/Core/Private/Windows/WindowsPlatformCrashContext.cpp
@@ -396,6 +396,18 @@ FProcHandle LaunchCrashReportClient(void** OutWritePipe, void** OutReadPipe)
 		FCString::Strncat(CrashReporterClientArgs, PidStr, CR_CLIENT_MAX_ARGS_LEN);
 	}
 
+#if WITH_EDITOR // Disaster recovery is only enabled for the Editor. Start the server even if in -game, -server, commandlet, the client-side will not connect (its too soon here to query this executable config).
+	{
+		// Disaster recovery service command line.
+		const FString DisasterRecoveryServerName = FString::Printf(TEXT("RecoverySvr_%d"), FPlatformProcess::GetCurrentProcessId());
+
+		FString DisasterRecoveryServiceCommandLine;
+		DisasterRecoveryServiceCommandLine += FString::Printf(TEXT(" -ConcertServer=\"%s\""), *DisasterRecoveryServerName); // Must be in-sync with disaster recovery client module.
+
+		FCString::Strncat(CrashReporterClientArgs, *DisasterRecoveryServiceCommandLine, CR_CLIENT_MAX_ARGS_LEN);
+	}
+#endif
+
 	// Launch the crash reporter if the client exists
 	if (CreateCrashReportClientPath(CrashReporterClientPath, CR_CLIENT_MAX_PATH_LEN))
 	{
