Removed request exit calls in crash reporter, fix logic for disaster recovery.

* Removed crash reporter client usage of "request exit", since we need to be able to send multiple reports in monitor mode. For example a ensure followed by a crash.
* Fix missing parenthesis when deciding to attach recovery service file.

diff --git a/Engine/Source/Programs/CrashReportClient/Private/CrashReportClient.cpp b/Engine/Source/Programs/CrashReportClient/Private/CrashReportClient.cpp
--- a/Engine/Source/Programs/CrashReportClient/Private/CrashReportClient.cpp
+++ b/Engine/Source/Programs/CrashReportClient/Private/CrashReportClient.cpp
@@ -38,6 +38,7 @@ FCrashReportClient::FCrashReportClient(const FPlatformErrorReport& InErrorReport
 	, bShouldWindowBeHidden(false)
 	, bSendData(false)
 	, bIsSuccesfullRestart(false)
+	, bIsUploadComplete(false)
 {
 	if (FPrimaryCrashProperties::Get()->IsValid())
 	{
@@ -286,7 +287,7 @@ bool FCrashReportClient::Tick(float UnusedDeltaTime)
 		FCrashUploadBase::StaticShutdown();
 	}
 
-	FPlatformMisc::RequestExit(false);
+	bIsUploadComplete = true;
 	return false;
 }
 
diff --git a/Engine/Source/Programs/CrashReportClient/Private/CrashReportClient.h b/Engine/Source/Programs/CrashReportClient/Private/CrashReportClient.h
--- a/Engine/Source/Programs/CrashReportClient/Private/CrashReportClient.h
+++ b/Engine/Source/Programs/CrashReportClient/Private/CrashReportClient.h
@@ -133,6 +133,8 @@ public:
 	void AllowToBeContacted_OnCheckStateChanged( ECheckBoxState NewRadioState );
 	void SendLogFile_OnCheckStateChanged( ECheckBoxState NewRadioState );
 
+	bool IsUploadComplete() const { return bIsUploadComplete; }
+
 private:
 	/**
 	 * Write the user's comment to the report and begin uploading the entire report 
@@ -189,6 +191,9 @@ private:
 	/** Store if user has elected to restart crashing process, and if it was a success.*/
 	bool bIsSuccesfullRestart;
 
+	/** Is the uploading complete. */
+	bool bIsUploadComplete;
+
 };
 
 #endif // !CRASH_REPORT_UNATTENDED_ONLY
diff --git a/Engine/Source/Programs/CrashReportClient/Private/CrashReportClientApp.cpp b/Engine/Source/Programs/CrashReportClient/Private/CrashReportClientApp.cpp
--- a/Engine/Source/Programs/CrashReportClient/Private/CrashReportClientApp.cpp
+++ b/Engine/Source/Programs/CrashReportClient/Private/CrashReportClientApp.cpp
@@ -376,7 +376,7 @@ SubmitCrashReportResult RunUnattended(FPlatformErrorReport ErrorReport)
 	}
 
 	// loop until the app is ready to quit
-	while (!IsEngineExitRequested())
+	while (!(IsEngineExitRequested() || CrashReportClient.IsUploadComplete()))
 	{
 		MainLoop.Tick();
 	}
@@ -470,7 +470,7 @@ FPlatformErrorReport CollectErrorReport(FRecoveryService* RecoveryService, uint3
 	const FString CrashContextXMLPath = FPaths::Combine(*ReportDirectoryAbsolutePath, FPlatformCrashContext::CrashContextRuntimeXMLNameW);
 	CrashContext.SerializeAsXML(*CrashContextXMLPath);
 
-	if (RecoveryService && DirectoryExists && SharedCrashContext.bSendUsageData && SharedCrashContext.CrashType != ECrashContextType::Ensure || SharedCrashContext.CrashType == ECrashContextType::Assert)
+	if (RecoveryService && DirectoryExists && SharedCrashContext.bSendUsageData && (SharedCrashContext.CrashType != ECrashContextType::Ensure || SharedCrashContext.CrashType == ECrashContextType::Assert))
 	{
 		RecoveryService->CollectFiles(ReportDirectoryAbsolutePath);
 	}
@@ -614,7 +614,7 @@ void RunCrashReportClient(const TCHAR* CommandLine)
 
 		// This IsApplicationAlive() call is quite expensive, perform it at low frequency.
 		bool bApplicationAlive = FPlatformProcess::IsApplicationAlive(MonitorPid);
-		while (bApplicationAlive && !GIsRequestingExit)
+		while (bApplicationAlive && !IsEngineExitRequested())
 		{
 			const double CurrentTime = FPlatformTime::Seconds();
 
@@ -652,10 +652,10 @@ void RunCrashReportClient(const TCHAR* CommandLine)
 					// This is ugly, but since many parts of CrashReportClient and CrashReportCoreUnattended is built
 					// using this flag to exit tick loops, this was the easiest way of making the monitor path be able to
 					// send multiple reports (e.g. ensure followed by crash).
-					if (Result == SuccessContinue)
+					/*if (Result == SuccessContinue)
 					{
 						GIsRequestingExit = false;
-					}
+					}*/
 
 					if (bReportCrashAnalyticInfo)
 					{
diff --git a/Engine/Source/Runtime/CrashReportCore/Private/CrashReportCoreUnattended.cpp b/Engine/Source/Runtime/CrashReportCore/Private/CrashReportCoreUnattended.cpp
--- a/Engine/Source/Runtime/CrashReportCore/Private/CrashReportCoreUnattended.cpp
+++ b/Engine/Source/Runtime/CrashReportCore/Private/CrashReportCoreUnattended.cpp
@@ -68,11 +68,6 @@ bool FCrashReportCoreUnattended::Tick(float UnusedDeltaTime)
 	}
 
     bUploadComplete = true;
-	
-	if (bExitWhenComplete)
-	{
-		FPlatformMisc::RequestExit(false /* don't force */);
-	}
 	return false;
 }
 
