Enabled disaster recovery session capture in case of crash.
Enabled capturing assert crash type.
Fixed the session name being looked for when collecting disaster recovery files.

diff --git a/Engine/Plugins/Developer/Concert/DisasterRecoveryClient/Source/DisasterRecoveryClient/Private/DisasterRecoveryClientModule.cpp b/Engine/Plugins/Developer/Concert/DisasterRecoveryClient/Source/DisasterRecoveryClient/Private/DisasterRecoveryClientModule.cpp
--- a/Engine/Plugins/Developer/Concert/DisasterRecoveryClient/Source/DisasterRecoveryClient/Private/DisasterRecoveryClientModule.cpp
+++ b/Engine/Plugins/Developer/Concert/DisasterRecoveryClient/Source/DisasterRecoveryClient/Private/DisasterRecoveryClientModule.cpp
@@ -392,7 +392,7 @@ private:
 		}
 
 		const FString DisasterRecoveryServerName = RecoveryService::GetRecoveryServerName();
-		const FString DisasterRecoverySessionName = FString::Printf(TEXT("%s_%s_%s"), FApp::GetProjectName(), *FDateTime::Now().ToString(), *DisasterRecoveryServerName);
+		const FString DisasterRecoverySessionName = FString::Printf(TEXT("%s_%s_%s"), *DisasterRecoveryServerName, FApp::GetProjectName(), *FDateTime::Now().ToString());
 
 		if (ShouldLaunchDisasterRecoveryBackendService() && !SpawnDisasterRecoveryServer(DisasterRecoveryServerName))
 		{
diff --git a/Engine/Source/Programs/CrashReportClient/Private/CrashReportClientApp.cpp b/Engine/Source/Programs/CrashReportClient/Private/CrashReportClientApp.cpp
--- a/Engine/Source/Programs/CrashReportClient/Private/CrashReportClientApp.cpp
+++ b/Engine/Source/Programs/CrashReportClient/Private/CrashReportClientApp.cpp
@@ -470,7 +470,7 @@ FPlatformErrorReport CollectErrorReport(FRecoveryService* RecoveryService, uint3
 	const FString CrashContextXMLPath = FPaths::Combine(*ReportDirectoryAbsolutePath, FPlatformCrashContext::CrashContextRuntimeXMLNameW);
 	CrashContext.SerializeAsXML(*CrashContextXMLPath);
 
-	if (RecoveryService && DirectoryExists && SharedCrashContext.bSendUsageData && (SharedCrashContext.CrashType != ECrashContextType::Ensure || SharedCrashContext.CrashType == ECrashContextType::Assert))
+	if (RecoveryService && DirectoryExists && SharedCrashContext.bSendUsageData && SharedCrashContext.CrashType != ECrashContextType::Ensure)
 	{
 		RecoveryService->CollectFiles(ReportDirectoryAbsolutePath);
 	}
@@ -633,8 +633,8 @@ void RunCrashReportClient(const TCHAR* CommandLine)
 						FCrashReportAnalytics::Initialize();
 					}
 
-					// Build error report in memory. (NOTE: Don't collect recovery service files, data anonymity is not implemented yet)
-					FPlatformErrorReport ErrorReport = CollectErrorReport(nullptr/*RecoveryService.Get()*/, MonitorPid, CrashContext, MonitorWritePipe);
+					// Build error report in memory.
+					FPlatformErrorReport ErrorReport = CollectErrorReport(RecoveryService.Get(), MonitorPid, CrashContext, MonitorWritePipe);
 					const SubmitCrashReportResult Result = SendErrorReport(ErrorReport, CrashContext.bNoDialog && CrashContext.bSendUnattenededBugReports);
 
 					// At this point the game can continue execution. It is important this happens
diff --git a/Engine/Source/Programs/CrashReportClient/Private/RecoveryService.cpp b/Engine/Source/Programs/CrashReportClient/Private/RecoveryService.cpp
--- a/Engine/Source/Programs/CrashReportClient/Private/RecoveryService.cpp
+++ b/Engine/Source/Programs/CrashReportClient/Private/RecoveryService.cpp
@@ -109,13 +109,11 @@ FString FRecoveryService::GetRecoveryServiceArchivedDir() const
 
 FGuid FRecoveryService::GetRecoverySessionId() const
 {
-	// As convention, the disaster recovery client names the session as Recover_{pid}_{projName}_{dateTime}.
-	FString SearchedSessionName = FString::Printf(TEXT("Recovery_%d"), MonitorPid);
-
 	// As long as the Concert server is up, the session would remain live (it's going to be archived when the server shutdown or reboot).
 	for (TSharedPtr<IConcertServerSession>& Session : Server->GetConcertServer()->GetSessions())
 	{
-		if (Session->GetName().StartsWith(SearchedSessionName))
+		// As convention, the disaster recovery session names starts with the server name, followed by the project name, and date time.
+		if (Session->GetName().StartsWith(Server->GetConcertServer()->GetServerInfo().ServerName))
 		{
 			return Session->GetId();
 		}
