- Explicitely enabled message bus if shipping if the build is monolithic because Concert/DisasterRecovery needs it to work.
  - Detected at runtime if message bus was unexpectedly disabled by build configuration and disabled recovery service rather than crashing on an failed assert.

diff --git a/Engine/Source/Programs/CrashReportClient/CrashReportClient.Build.cs b/Engine/Source/Programs/CrashReportClient/CrashReportClient.Build.cs
--- a/Engine/Source/Programs/CrashReportClient/CrashReportClient.Build.cs
+++ b/Engine/Source/Programs/CrashReportClient/CrashReportClient.Build.cs
@@ -39,6 +39,7 @@ public class CrashReportClient : ModuleRules
 				"SlateCore",
 				"StandaloneRenderer",
 				"MessageLog",
+				"Messaging",
 			}
 		);
 
diff --git a/Engine/Source/Programs/CrashReportClient/CrashReportClient.Target.cs b/Engine/Source/Programs/CrashReportClient/CrashReportClient.Target.cs
--- a/Engine/Source/Programs/CrashReportClient/CrashReportClient.Target.cs
+++ b/Engine/Source/Programs/CrashReportClient/CrashReportClient.Target.cs
@@ -25,7 +25,7 @@ public class CrashReportClientTarget : TargetRules
 			ExtraModuleNames.Add("EditorStyle");
 		}
 
-        bLegalToDistributeBinary = true;
+		bLegalToDistributeBinary = true;
 
 		bBuildDeveloperTools = true;
 
@@ -47,5 +47,12 @@ public class CrashReportClientTarget : TargetRules
 		WindowsPlatform.bUseBundledDbgHelp = false;
 
 		GlobalDefinitions.Add("NOINITCRASHREPORTER=1");
+
+		if (Target.Configuration == UnrealTargetConfiguration.Shipping && LinkType == TargetLinkType.Monolithic)
+		{
+			// DisasterRecovery/Concert needs message bus to run. If not enabled, Recovery Service will self-disable as well. In Shipping
+			// message bus is turned off by default but for a monolithic build, it can be turned on just for this executable.
+			GlobalDefinitions.Add("PLATFORM_SUPPORTS_MESSAGEBUS=1");
+		}
 	}
 }
diff --git a/Engine/Source/Programs/CrashReportClient/Private/RecoveryService.cpp b/Engine/Source/Programs/CrashReportClient/Private/RecoveryService.cpp
--- a/Engine/Source/Programs/CrashReportClient/Private/RecoveryService.cpp
+++ b/Engine/Source/Programs/CrashReportClient/Private/RecoveryService.cpp
@@ -4,6 +4,7 @@
 #include "HAL/FileManager.h"
 #include "CrashReportClient.h" // For CrashReportClientLog
 #include "Interfaces/IPluginManager.h"
+#include "IMessagingModule.h"
 #include "ConcertSettings.h"
 #include "ConcertSyncSessionFlags.h"
 #include "IConcertServer.h"
@@ -24,7 +25,7 @@ bool FRecoveryService::CollectFiles(const FString& DestDir, bool bMetaDataOnly,
 
 	if (!Server)
 	{
-		LogError(TEXT("The recovery service is not running anymore."));
+		LogError(TEXT("The recovery service is not running."));
 		return false;
 	}
 	else if (!IFileManager::Get().DirectoryExists(*DestDir))
@@ -54,6 +55,15 @@ bool FRecoveryService::CollectFiles(const FString& DestDir, bool bMetaDataOnly,
 
 bool FRecoveryService::Startup()
 {
+#if UE_BUILD_SHIPPING && (!defined(PLATFORM_SUPPORTS_MESSAGEBUS) || !PLATFORM_SUPPORTS_MESSAGEBUS)
+	#error PLATFORM_SUPPORTS_MESSAGEBUS was explicitly defined in CrashReportClient.Target.cs for shipping configuration. MessageBus is required by Concert. Ensure it is still enabled.
+#endif
+	if (!IMessagingModule::Get().GetDefaultBus())
+	{
+		UE_LOG(CrashReportClientLog, Error, TEXT("MessageBus is not enabled in this configuration. Recovery service will be disabled!"));
+		return false;
+	}
+
 	if (!IConcertSyncServerModule::IsAvailable())
 	{
 		UE_LOG(CrashReportClientLog, Error, TEXT("ConcertSyncServer Module is missing. Recovery service will be disabled!"));
