Add target for editor crash report client

There are some operations in the editor (package and stage) that needs to build Crash report client, but when we run it in monitor mode the binary is locked failing the build. This change adds a new target (basically a copy of the crash report client) for the editor. This also allows us to add specific dependencies for the editor client (such as disaster recovery) which should not be shipped with the regular version.

diff --git a/Engine/Source/Programs/CrashReportClient/CrashReportClientEditor.Build.cs b/Engine/Source/Programs/CrashReportClient/CrashReportClientEditor.Build.cs
new file mode 100644
--- /dev/null
+++ b/Engine/Source/Programs/CrashReportClient/CrashReportClientEditor.Build.cs
@@ -0,0 +1,10 @@
+// Copyright 1998-2019 Epic Games, Inc. All Rights Reserved.
+
+using UnrealBuildTool;
+
+public class CrashReportClientEditor : CrashReportClient
+{
+	public CrashReportClientEditor(ReadOnlyTargetRules Target) : base(Target)
+	{
+	}
+}
diff --git a/Engine/Source/Programs/CrashReportClient/CrashReportClientEditor.Target.cs b/Engine/Source/Programs/CrashReportClient/CrashReportClientEditor.Target.cs
new file mode 100644
--- /dev/null
+++ b/Engine/Source/Programs/CrashReportClient/CrashReportClientEditor.Target.cs
@@ -0,0 +1,14 @@
+// Copyright 1998-2019 Epic Games, Inc. All Rights Reserved.
+
+using UnrealBuildTool;
+using System.Collections.Generic;
+
+[SupportedPlatforms("Win64", "Mac", "Linux")]
+[SupportedConfigurations(UnrealTargetConfiguration.Debug, UnrealTargetConfiguration.Development, UnrealTargetConfiguration.Shipping)]
+public class CrashReportClientEditorTarget : CrashReportClientTarget
+{
+    public CrashReportClientEditorTarget(TargetInfo Target) : base(Target)
+    {
+        LaunchModuleName = "CrashReportClientEditor";
+    }
+}
diff --git a/Engine/Source/Runtime/Core/Private/Windows/WindowsPlatformCrashContext.cpp b/Engine/Source/Runtime/Core/Private/Windows/WindowsPlatformCrashContext.cpp
--- a/Engine/Source/Runtime/Core/Private/Windows/WindowsPlatformCrashContext.cpp
+++ b/Engine/Source/Runtime/Core/Private/Windows/WindowsPlatformCrashContext.cpp
@@ -363,6 +363,13 @@ bool CreateCrashReportClientPath(TCHAR* OutClientPath, int32 MaxLength)
 		return Results != INVALID_FILE_ATTRIBUTES;
 	};
 
+#if WITH_EDITORONLY_DATA
+	if (CreateCrashReportClientPathImpl(TEXT("CrashReportClientEditor.exe")))
+	{
+		return true;
+	}
+#endif
+
 	if (CreateCrashReportClientPathImpl(TEXT("CrashReportClient.exe")))
 	{
 		return true;
