Fix restart editor when running crash reporter with GUI.

Add the same check to break out of the main loop when the upload is complete as is being used in the RunUnattended method.

diff --git a/Engine/Source/Programs/CrashReportClient/Private/CrashReportClientApp.cpp b/Engine/Source/Programs/CrashReportClient/Private/CrashReportClientApp.cpp
--- a/Engine/Source/Programs/CrashReportClient/Private/CrashReportClientApp.cpp
+++ b/Engine/Source/Programs/CrashReportClient/Private/CrashReportClientApp.cpp
@@ -318,7 +318,7 @@ SubmitCrashReportResult RunWithUI(FPlatformErrorReport ErrorReport)
 	}
 
 	// loop until the app is ready to quit
-	while (!IsEngineExitRequested())
+	while (!(IsEngineExitRequested() || CrashReportClient->IsUploadComplete()))
 	{
 		MainLoop.Tick();
 
@@ -660,14 +660,6 @@ void RunCrashReportClient(const TCHAR* CommandLine)
 						// If analytics is enabled make sure they are submitted now.
 						FCrashReportAnalytics::GetProvider().BlockUntilFlushed(5.0f);
 					}
-
-					// This is ugly, but since many parts of CrashReportClient and CrashReportCoreUnattended is built
-					// using this flag to exit tick loops, this was the easiest way of making the monitor path be able to
-					// send multiple reports (e.g. ensure followed by crash).
-					/*if (Result == SuccessContinue)
-					{
-						GIsRequestingExit = false;
-					}*/
 				}
 			}
 
