Set Disaster Recovery and MTBF reporting to only be enabled for the editor variant of the Crash Report Client

#rnx

diff --git a/Engine/Programs/CrashReportClient/Config/DefaultEngine.ini b/Engine/Programs/CrashReportClient/Config/DefaultEngine.ini
--- a/Engine/Programs/CrashReportClient/Config/DefaultEngine.ini
+++ b/Engine/Programs/CrashReportClient/Config/DefaultEngine.ini
@@ -10,10 +10,6 @@ UIInitRetryInterval=2.0
 [/Script/Engine.NetworkSettings]
 n.VerifyPeer=false
 
-[Plugins]
-+ProgramEnabledPlugins="UdpMessaging"
-+ProgramEnabledPlugins="ConcertSyncServer"
-
 [/Script/UdpMessaging.UdpMessagingSettings]
 EnableTransport=True
 MulticastTimeToLive=0
diff --git a/Engine/Source/Programs/CrashReportClient/CrashReportClient.Build.cs b/Engine/Source/Programs/CrashReportClient/CrashReportClient.Build.cs
--- a/Engine/Source/Programs/CrashReportClient/CrashReportClient.Build.cs
+++ b/Engine/Source/Programs/CrashReportClient/CrashReportClient.Build.cs
@@ -32,29 +32,23 @@ public class CrashReportClient : ModuleRules
 				"AnalyticsET",
 				"DesktopPlatform",
 				"LauncherPlatform",
-				"Concert",
 				"InputCore",
-				"EditorAnalyticsSession",
 				"Slate",
 				"SlateCore",
 				"StandaloneRenderer",
 				"MessageLog",
-				"Messaging",
 			}
 		);
 
 		PrivateIncludePathModuleNames.AddRange(
 			new string[] {
 				"SlateReflector",
-				"ConcertSyncCore",
-				"ConcertSyncServer",
 			}
 		);
 
 		DynamicallyLoadedModuleNames.AddRange(
 			new string[] {
 				"SlateReflector",
-				"ConcertSyncServer",
 			}
 		);
 
diff --git a/Engine/Source/Programs/CrashReportClient/CrashReportClient.Target.cs b/Engine/Source/Programs/CrashReportClient/CrashReportClient.Target.cs
--- a/Engine/Source/Programs/CrashReportClient/CrashReportClient.Target.cs
+++ b/Engine/Source/Programs/CrashReportClient/CrashReportClient.Target.cs
@@ -15,11 +15,6 @@ public class CrashReportClientTarget : TargetRules
 
 		LaunchModuleName = "CrashReportClient";
 
-		AdditionalPlugins.Add("UdpMessaging");
-		AdditionalPlugins.Add("ConcertSyncServer");
-		bBuildWithEditorOnlyData = false;
-		bCompileWithPluginSupport = true; // Enable Developer plugins (like Concert!)
-
 		if (Target.Platform != UnrealTargetPlatform.Linux)
 		{
 			ExtraModuleNames.Add("EditorStyle");
@@ -27,7 +22,7 @@ public class CrashReportClientTarget : TargetRules
 
 		bLegalToDistributeBinary = true;
 
-		bBuildDeveloperTools = true;
+		bBuildDeveloperTools = false;
 
 		// CrashReportClient doesn't ever compile with the engine linked in
 		bCompileAgainstEngine = false;
@@ -47,12 +42,5 @@ public class CrashReportClientTarget : TargetRules
 		WindowsPlatform.bUseBundledDbgHelp = false;
 
 		GlobalDefinitions.Add("NOINITCRASHREPORTER=1");
-
-		if (Target.Configuration == UnrealTargetConfiguration.Shipping && LinkType == TargetLinkType.Monolithic)
-		{
-			// DisasterRecovery/Concert needs message bus to run. If not enabled, Recovery Service will self-disable as well. In Shipping
-			// message bus is turned off by default but for a monolithic build, it can be turned on just for this executable.
-			GlobalDefinitions.Add("PLATFORM_SUPPORTS_MESSAGEBUS=1");
-		}
 	}
 }
diff --git a/Engine/Source/Programs/CrashReportClient/CrashReportClientEditor.Build.cs b/Engine/Source/Programs/CrashReportClient/CrashReportClientEditor.Build.cs
--- a/Engine/Source/Programs/CrashReportClient/CrashReportClientEditor.Build.cs
+++ b/Engine/Source/Programs/CrashReportClient/CrashReportClientEditor.Build.cs
@@ -6,5 +6,35 @@ public class CrashReportClientEditor : CrashReportClient
 {
 	public CrashReportClientEditor(ReadOnlyTargetRules Target) : base(Target)
 	{
+		PrivateDefinitions.AddRange(
+			new string[]
+			{
+				"CRASH_REPORT_WITH_RECOVERY=1",
+				"CRASH_REPORT_WITH_MTBF=1",
+			}
+		);
+
+		PrivateDependencyModuleNames.AddRange(
+			new string[]
+			{
+				"Concert",
+				"EditorAnalyticsSession",
+				"Messaging",
+			}
+		);
+
+		PrivateIncludePathModuleNames.AddRange(
+			new string[] {
+				"ConcertSyncCore",
+				"ConcertSyncServer",
+			}
+		);
+
+		DynamicallyLoadedModuleNames.AddRange(
+			new string[] {
+				"ConcertSyncServer",
+				"UdpMessaging",
+			}
+		);
 	}
 }
diff --git a/Engine/Source/Programs/CrashReportClient/CrashReportClientEditor.Target.cs b/Engine/Source/Programs/CrashReportClient/CrashReportClientEditor.Target.cs
--- a/Engine/Source/Programs/CrashReportClient/CrashReportClientEditor.Target.cs
+++ b/Engine/Source/Programs/CrashReportClient/CrashReportClientEditor.Target.cs
@@ -10,5 +10,19 @@ public class CrashReportClientEditorTarget : CrashReportClientTarget
     public CrashReportClientEditorTarget(TargetInfo Target) : base(Target)
     {
         LaunchModuleName = "CrashReportClientEditor";
-    }
+
+		AdditionalPlugins.Add("UdpMessaging");
+		AdditionalPlugins.Add("ConcertSyncServer");
+		bBuildWithEditorOnlyData = false;
+		bCompileWithPluginSupport = true; // Enable Developer plugins (like Concert!)
+
+		bBuildDeveloperTools = true;
+
+		if (Target.Configuration == UnrealTargetConfiguration.Shipping && LinkType == TargetLinkType.Monolithic)
+		{
+			// DisasterRecovery/Concert needs message bus to run. If not enabled, Recovery Service will self-disable as well. In Shipping
+			// message bus is turned off by default but for a monolithic build, it can be turned on just for this executable.
+			GlobalDefinitions.Add("PLATFORM_SUPPORTS_MESSAGEBUS=1");
+		}
+	}
 }
diff --git a/Engine/Source/Programs/CrashReportClient/Private/CrashReportClientApp.cpp b/Engine/Source/Programs/CrashReportClient/Private/CrashReportClientApp.cpp
--- a/Engine/Source/Programs/CrashReportClient/Private/CrashReportClientApp.cpp
+++ b/Engine/Source/Programs/CrashReportClient/Private/CrashReportClientApp.cpp
@@ -1,6 +1,7 @@
 // Copyright 1998-2019 Epic Games, Inc. All Rights Reserved.
 
 #include "CrashReportClientApp.h"
+#include "CrashReportClientDefines.h"
 #include "Misc/Parse.h"
 #include "Misc/CommandLine.h"
 #include "Misc/QueuedThreadPool.h"
@@ -35,7 +36,12 @@
 #include "PlatformErrorReport.h"
 #include "XmlFile.h"
 #include "RecoveryService.h"
+
+#if CRASH_REPORT_WITH_MTBF
 #include "EditorSessionSummarySender.h"
+#endif
+
+class FRecoveryService;
 
 /** Default main window size */
 const FVector2D InitialWindowDimensions(740, 560);
@@ -72,12 +78,6 @@ static void* MonitorReadPipe = nullptr;
 /** If in monitor mode, pipe to write data to game. */
 static void* MonitorWritePipe = nullptr;
 
-/** True to enable the disaster recovery service. */
-static bool bDisasterRecoveryServiceEnabled = false;
-
-/** Is the out-of-process editor session summary sender enabled? */
-static bool bEditorSessionSummaryEnabled = true;
-
 /** Result of submission of report */
 enum SubmitCrashReportResult {
 	Failed,				// Failed to send report
@@ -150,11 +150,6 @@ void ParseCommandLine(const TCHAR* CommandLine)
 			AnalyticsEnabledFromCmd = false;
 		}
 
-		if (Params.Contains(TEXT("ConcertServer")))
-		{
-			bDisasterRecoveryServiceEnabled = true;
-		}
-
 		CrashGUIDFromCmd = Params.FindRef(TEXT("CrashGUID"));
 		MonitorPid = FPlatformString::Atoi64(*Params.FindRef(TEXT("MONITOR")));
 		MonitorReadPipe = (void*) FPlatformString::Atoi64(*Params.FindRef(TEXT("READ")));
@@ -474,10 +469,12 @@ FPlatformErrorReport CollectErrorReport(FRecoveryService* RecoveryService, uint3
 	const FString CrashContextXMLPath = FPaths::Combine(*ReportDirectoryAbsolutePath, FPlatformCrashContext::CrashContextRuntimeXMLNameW);
 	CrashContext.SerializeAsXML(*CrashContextXMLPath);
 
+#if CRASH_REPORT_WITH_RECOVERY
 	if (RecoveryService && DirectoryExists && SharedCrashContext.bSendUsageData && SharedCrashContext.CrashType != ECrashContextType::Ensure)
 	{
 		RecoveryService->CollectFiles(ReportDirectoryAbsolutePath);
 	}
+#endif
 
 	const TCHAR* CrachContextBuffer = *CrashContext.GetBuffer();
 	FPrimaryCrashProperties::Set(new FCrashContext(ReportDirectoryAbsolutePath / TEXT("CrashContext.runtime-xml"), CrachContextBuffer));
@@ -562,7 +559,9 @@ void RunCrashReportClient(const TCHAR* CommandLine)
 
 	// Initialize the engine. -Messaging enables MessageBus transports required by Concert (Recovery Service).
 	FString FinalCommandLine(CommandLine);
-	FinalCommandLine += TEXT(" -Messaging");
+#if CRASH_REPORT_WITH_RECOVERY
+	FinalCommandLine += TEXT(" -Messaging -EnablePlugins=\"UdpMessaging,ConcertSyncServer\"");
+#endif
 	GEngineLoop.PreInit(*FinalCommandLine);
 	check(GConfig && GConfig->IsReadyForUse());
 
@@ -609,15 +608,16 @@ void RunCrashReportClient(const TCHAR* CommandLine)
 		double LastTime = FPlatformTime::Seconds();
 		const float IdealFrameTime = 1.0f / IdealFramerate;
 
-		TUniquePtr<FRecoveryService> RecoveryService;
-		if (bDisasterRecoveryServiceEnabled)
-		{
-			// Starts the disaster recovery service. This records transactions and allows users to recover from previous crashes.
-			RecoveryService = MakeUnique<FRecoveryService>(MonitorPid);
-		}
+		FRecoveryService* RecoveryServicePtr = nullptr;
+#if CRASH_REPORT_WITH_RECOVERY
+		// Starts the disaster recovery service. This records transactions and allows users to recover from previous crashes.
+		FRecoveryService RecoveryService(MonitorPid);
+		RecoveryServicePtr = &RecoveryService;
+#endif
 
 		FCrashReportAnalytics::Initialize();
 
+#if CRASH_REPORT_WITH_MTBF
 		TUniquePtr<FEditorSessionSummarySender> EditorSessionSummarySender;
 		if (FCrashReportCoreConfig::Get().GetAllowToBeContacted())
 		{
@@ -630,6 +630,7 @@ void RunCrashReportClient(const TCHAR* CommandLine)
 					return true;
 				});
 		}
+#endif
 
 		// This IsApplicationAlive() call is quite expensive, perform it at low frequency.
 		bool bApplicationAlive = FPlatformProcess::IsApplicationAlive(MonitorPid);
@@ -645,7 +646,7 @@ void RunCrashReportClient(const TCHAR* CommandLine)
 				if (IsCrashReportAvailable(MonitorPid, CrashContext, MonitorReadPipe))
 				{
 					// Build error report in memory.
-					FPlatformErrorReport ErrorReport = CollectErrorReport(RecoveryService.Get(), MonitorPid, CrashContext, MonitorWritePipe);
+					FPlatformErrorReport ErrorReport = CollectErrorReport(RecoveryServicePtr, MonitorPid, CrashContext, MonitorWritePipe);
 					const SubmitCrashReportResult Result = SendErrorReport(ErrorReport, CrashContext.bNoDialog && CrashContext.bSendUnattenededBugReports);
 
 					// At this point the game can continue execution. It is important this happens
@@ -688,10 +689,12 @@ void RunCrashReportClient(const TCHAR* CommandLine)
 			LastTime = CurrentTime;
 		}
 
+#if CRASH_REPORT_WITH_MTBF
 		if (EditorSessionSummarySender.IsValid())
 		{
 			EditorSessionSummarySender->Shutdown();
 		}
+#endif
 
 		FCrashReportAnalytics::Shutdown();
 	}
diff --git a/Engine/Source/Programs/CrashReportClient/Private/CrashReportClientDefines.h b/Engine/Source/Programs/CrashReportClient/Private/CrashReportClientDefines.h
new file mode 100644
--- /dev/null
+++ b/Engine/Source/Programs/CrashReportClient/Private/CrashReportClientDefines.h
@@ -0,0 +1,11 @@
+// Copyright 1998-2019 Epic Games, Inc. All Rights Reserved.
+
+#pragma once
+
+#ifndef CRASH_REPORT_WITH_RECOVERY
+#define CRASH_REPORT_WITH_RECOVERY 0
+#endif
+
+#ifndef CRASH_REPORT_WITH_MTBF
+#define CRASH_REPORT_WITH_MTBF 0
+#endif
diff --git a/Engine/Source/Programs/CrashReportClient/Private/RecoveryService.cpp b/Engine/Source/Programs/CrashReportClient/Private/RecoveryService.cpp
--- a/Engine/Source/Programs/CrashReportClient/Private/RecoveryService.cpp
+++ b/Engine/Source/Programs/CrashReportClient/Private/RecoveryService.cpp
@@ -1,6 +1,9 @@
 // Copyright 1998-2019 Epic Games, Inc. All Rights Reserved.
 
 #include "RecoveryService.h"
+
+#if CRASH_REPORT_WITH_RECOVERY
+
 #include "HAL/FileManager.h"
 #include "CrashReportClient.h" // For CrashReportClientLog
 #include "Interfaces/IPluginManager.h"
@@ -131,3 +134,5 @@ FGuid FRecoveryService::GetRecoverySessionId() const
 
 	return FGuid();// Not found.
 }
+
+#endif
diff --git a/Engine/Source/Programs/CrashReportClient/Private/RecoveryService.h b/Engine/Source/Programs/CrashReportClient/Private/RecoveryService.h
--- a/Engine/Source/Programs/CrashReportClient/Private/RecoveryService.h
+++ b/Engine/Source/Programs/CrashReportClient/Private/RecoveryService.h
@@ -3,6 +3,9 @@
 #pragma once
 
 #include "CoreMinimal.h"
+#include "CrashReportClientDefines.h"
+
+#if CRASH_REPORT_WITH_RECOVERY
 
 class IConcertSyncServer;
 
@@ -28,3 +31,5 @@ private:
 	TSharedPtr<IConcertSyncServer> Server;
 	uint64 MonitorPid = 0;
 };
+
+#endif
diff --git a/Engine/Source/Runtime/Core/Private/Windows/WindowsPlatformCrashContext.cpp b/Engine/Source/Runtime/Core/Private/Windows/WindowsPlatformCrashContext.cpp
--- a/Engine/Source/Runtime/Core/Private/Windows/WindowsPlatformCrashContext.cpp
+++ b/Engine/Source/Runtime/Core/Private/Windows/WindowsPlatformCrashContext.cpp
@@ -363,20 +363,21 @@ bool CreateCrashReportClientPath(TCHAR* OutClientPath, int32 MaxLength)
 		return Results != INVALID_FILE_ATTRIBUTES;
 	};
 
-#if WITH_EDITORONLY_DATA
-	if (CreateCrashReportClientPathImpl(TEXT("CrashReportClientEditor.exe")))
-	{
-		return true;
-	}
+#if WITH_EDITOR
+	const TCHAR CrashReportClientShippingName[] = TEXT("CrashReportClientEditor.exe");
+	const TCHAR CrashReportClientDevelopmentName[] = TEXT("CrashReportClientEditor-Win64-Development.exe");
+#else
+	const TCHAR CrashReportClientShippingName[] = TEXT("CrashReportClient.exe");
+	const TCHAR CrashReportClientDevelopmentName[] = TEXT("CrashReportClient-Win64-Development.exe");
 #endif
 
-	if (CreateCrashReportClientPathImpl(TEXT("CrashReportClient.exe")))
+	if (CreateCrashReportClientPathImpl(CrashReportClientShippingName))
 	{
 		return true;
 	}
 
 #if !(UE_BUILD_TEST || UE_BUILD_SHIPPING)
-	if (CreateCrashReportClientPathImpl(TEXT("CrashReportClient-Win64-Development.exe")))
+	if (CreateCrashReportClientPathImpl(CrashReportClientDevelopmentName))
 	{
 		return true;
 	}
