Get the monitored process return code in the Crash Report Client and report it via MTBF analytics

 - Removed FPlatformProcess::IsApplicationAlive as it's just a combination of (FPlatformProcess::IsProcRunning(...) && !FPlatformProcess::GetProcReturnCode(...)), and the latter lets you get the return code of the monitored application.
   - Note: Getting the return code is only supported on Windows.
 - Updated FWindowsPlatformProcess::GetProcReturnCode to leave ReturnCode unchanged if the function returns false.
   - This is consistent with Mac and Linux.
 - Updated CrashReportClientApp to open the monitor PID during init and use that handle for the rest of the run.
   - This avoids issues with PID re-use as you have a persistent handle to the process, and also avoids repeatedly opening a process handle.
   - Note: This has no benefit on Mac or Linux, as FProcHandle is just a wrapper around a PID.

#rnx

diff --git a/Engine/Source/Developer/EditorAnalyticsSession/Private/EditorSessionSummarySender.cpp b/Engine/Source/Developer/EditorAnalyticsSession/Private/EditorSessionSummarySender.cpp
--- a/Engine/Source/Developer/EditorAnalyticsSession/Private/EditorSessionSummarySender.cpp
+++ b/Engine/Source/Developer/EditorAnalyticsSession/Private/EditorSessionSummarySender.cpp
@@ -26,10 +26,11 @@ namespace EditorSessionSenderDefs
 	static const FString AbnormalSessionToken(TEXT("AbnormalShutdown"));
 }
 
-FEditorSessionSummarySender::FEditorSessionSummarySender(IAnalyticsProvider& InAnalyticsProvider, const FString& InSenderName) :
-	HeartbeatTimeElapsed(0.0f)
+FEditorSessionSummarySender::FEditorSessionSummarySender(IAnalyticsProvider& InAnalyticsProvider, const FString& InSenderName, const int32 InCurrentSessionProcessId)
+	: HeartbeatTimeElapsed(0.0f)
 	, AnalyticsProvider(InAnalyticsProvider)
 	, Sender(InSenderName)
+	, CurrentSessionProcessId(InCurrentSessionProcessId)
 {
 	SendStoredSessions();
 }
@@ -48,10 +49,16 @@ void FEditorSessionSummarySender::Tick(float DeltaTime)
 
 void FEditorSessionSummarySender::Shutdown()
 {
-	SendStoredSessions();
+	SendStoredSessions(/*bForceSendCurrentSession*/true);
+}
+
+void FEditorSessionSummarySender::SetCurrentSessionExitCode(const int32 InCurrentSessionProcessId, const int32 InExitCode)
+{
+	check(CurrentSessionProcessId == InCurrentSessionProcessId);
+	CurrentSessionExitCode = InExitCode;
 }
 
-void FEditorSessionSummarySender::SendStoredSessions() const
+void FEditorSessionSummarySender::SendStoredSessions(const bool bForceSendCurrentSession) const
 {
 	TArray<FEditorAnalyticsSession> SessionsToReport;
 
@@ -66,7 +73,8 @@ void FEditorSessionSummarySender::SendStoredSessions() const
 		// Check each stored session to see if they should be sent or not 
 		for (FEditorAnalyticsSession& Session : ExistingSessions)
 		{
-			if (FPlatformProcess::IsApplicationRunning(Session.PlatformProcessID))
+			const bool bForceSendSession = bForceSendCurrentSession && (Session.PlatformProcessID == CurrentSessionProcessId);
+			if (!bForceSendSession && FPlatformProcess::IsApplicationRunning(Session.PlatformProcessID))
 			{
 				// Skip processes that are still running
 				continue;
@@ -169,5 +177,10 @@ void FEditorSessionSummarySender::SendSessionSummaryEvent(const FEditorAnalytics
 
 	AnalyticsProvider.RecordEvent(TEXT("SessionSummary"), AnalyticsAttributes);
 
+	if (Session.PlatformProcessID == CurrentSessionProcessId && CurrentSessionExitCode.IsSet())
+	{
+		AnalyticsAttributes.Emplace(TEXT("ExitCode"), CurrentSessionExitCode.GetValue());
+	}
+
 	UE_LOG(LogEditorSessionSummary, Log, TEXT("EditorSessionSummary sent report. Type=%s, SessionId=%s"), *ShutdownTypeString, *SessionIdString);
 }
diff --git a/Engine/Source/Developer/EditorAnalyticsSession/Public/EditorSessionSummarySender.h b/Engine/Source/Developer/EditorAnalyticsSession/Public/EditorSessionSummarySender.h
--- a/Engine/Source/Developer/EditorAnalyticsSession/Public/EditorSessionSummarySender.h
+++ b/Engine/Source/Developer/EditorAnalyticsSession/Public/EditorSessionSummarySender.h
@@ -15,18 +15,23 @@ class IAnalyticsProvider;
 class EDITORANALYTICSSESSION_API FEditorSessionSummarySender
 {
 public:
-	FEditorSessionSummarySender(IAnalyticsProvider& InAnalyticsProvider, const FString& InSenderName);
+	FEditorSessionSummarySender(IAnalyticsProvider& InAnalyticsProvider, const FString& InSenderName, const int32 InCurrentSessionProcessId);
 
 	void Tick(float DeltaTime);
 	void Shutdown();
 
+	void SetCurrentSessionExitCode(const int32 InCurrentSessionProcessId, const int32 InExitCode);
+
 private:
 	/** Send any stored Sessions. */
-	void SendStoredSessions() const;
+	void SendStoredSessions(const bool bForceSendCurrentSession = false) const;
 	void SendSessionSummaryEvent(const FEditorAnalyticsSession& Session) const;
 
 private:
 	float HeartbeatTimeElapsed;
 	IAnalyticsProvider& AnalyticsProvider;
 	FString Sender;
-};
\ No newline at end of file
+
+	int32 CurrentSessionProcessId;
+	TOptional<int32> CurrentSessionExitCode;
+};
diff --git a/Engine/Source/Programs/CrashReportClient/Private/CrashReportClientApp.cpp b/Engine/Source/Programs/CrashReportClient/Private/CrashReportClientApp.cpp
--- a/Engine/Source/Programs/CrashReportClient/Private/CrashReportClientApp.cpp
+++ b/Engine/Source/Programs/CrashReportClient/Private/CrashReportClientApp.cpp
@@ -312,6 +312,10 @@ SubmitCrashReportResult RunWithUI(FPlatformErrorReport ErrorReport)
 		FModuleManager::LoadModuleChecked<ISlateReflectorModule>("SlateReflector").DisplayWidgetReflector();
 	}
 
+	// Bring the window to the foreground as it may be behind the crashed process
+	Window->HACK_ForceToFront();
+	Window->BringToFront();
+
 	// loop until the app is ready to quit
 	while (!(IsEngineExitRequested() || CrashReportClient->IsUploadComplete()))
 	{
@@ -621,19 +625,31 @@ void RunCrashReportClient(const TCHAR* CommandLine)
 		TUniquePtr<FEditorSessionSummarySender> EditorSessionSummarySender;
 		if (FCrashReportCoreConfig::Get().GetAllowToBeContacted())
 		{
-			EditorSessionSummarySender = MakeUnique<FEditorSessionSummarySender>(FCrashReportAnalytics::GetProvider(), TEXT("CrashReportClient"));
+			EditorSessionSummarySender = MakeUnique<FEditorSessionSummarySender>(FCrashReportAnalytics::GetProvider(), TEXT("CrashReportClient"), MonitorPid);
 
-			FTicker::GetCoreTicker().AddTicker(TEXT("EditorSessionSummarySender"), 0, 
-				[&EditorSessionSummarySender](float DeltaTime)
-				{
-					EditorSessionSummarySender->Tick(DeltaTime);
-					return true;
-				});
+			FTicker::GetCoreTicker().AddTicker(TEXT("EditorSessionSummarySender"), 0, [&EditorSessionSummarySender](float DeltaTime)
+			{
+				EditorSessionSummarySender->Tick(DeltaTime);
+				return true;
+			});
 		}
 #endif
 
+		FProcHandle MonitoredProcess = FPlatformProcess::OpenProcess(MonitorPid);
+		if (!MonitoredProcess.IsValid())
+		{
+			UE_LOG(CrashReportClientLog, Error, TEXT("Failed to open monitor process handle!"));
+		}
+
+		auto IsMonitoredProcessAlive = [&MonitoredProcess](int32& OutReturnCode) -> bool
+		{
+			// The monitored process is considered "alive" if the process is running, and no return code has been set
+			return MonitoredProcess.IsValid() && FPlatformProcess::IsProcRunning(MonitoredProcess) && !FPlatformProcess::GetProcReturnCode(MonitoredProcess, &OutReturnCode);
+		};
+
 		// This IsApplicationAlive() call is quite expensive, perform it at low frequency.
-		bool bApplicationAlive = FPlatformProcess::IsApplicationAlive(MonitorPid);
+		int32 ApplicationReturnCode = 0;
+		bool bApplicationAlive = IsMonitoredProcessAlive(ApplicationReturnCode);
 		while (bApplicationAlive && !IsEngineExitRequested())
 		{
 			const double CurrentTime = FPlatformTime::Seconds();
@@ -683,7 +699,7 @@ void RunCrashReportClient(const TCHAR* CommandLine)
 			// Check if the application is alive about every second. (This is an expensive call)
 			if (GFrameCounter % IdealFramerate == 0)
 			{
-				bApplicationAlive = FPlatformProcess::IsApplicationAlive(MonitorPid);
+				bApplicationAlive = IsMonitoredProcessAlive(ApplicationReturnCode);
 			}
 
 			LastTime = CurrentTime;
@@ -692,6 +708,12 @@ void RunCrashReportClient(const TCHAR* CommandLine)
 #if CRASH_REPORT_WITH_MTBF
 		if (EditorSessionSummarySender.IsValid())
 		{
+			// Query this again, as the crash reporting loop above may have exited before setting this information (via IsEngineExitRequested)
+			bApplicationAlive = IsMonitoredProcessAlive(ApplicationReturnCode);
+			if (!bApplicationAlive)
+			{
+				EditorSessionSummarySender->SetCurrentSessionExitCode(MonitorPid, ApplicationReturnCode);
+			}
 			EditorSessionSummarySender->Shutdown();
 		}
 #endif
diff --git a/Engine/Source/Runtime/Core/Private/GenericPlatform/GenericPlatformProcess.cpp b/Engine/Source/Runtime/Core/Private/GenericPlatform/GenericPlatformProcess.cpp
--- a/Engine/Source/Runtime/Core/Private/GenericPlatform/GenericPlatformProcess.cpp
+++ b/Engine/Source/Runtime/Core/Private/GenericPlatform/GenericPlatformProcess.cpp
@@ -284,12 +284,6 @@ bool FGenericPlatformProcess::IsApplicationRunning( const TCHAR* ProcName )
 	return false;
 }
 
-bool FGenericPlatformProcess::IsApplicationAlive(uint32 ProcessId)
-{
-	UE_LOG(LogHAL, Fatal, TEXT("FGenericPlatformProcess::IsApplicationAlive not implemented on this platform"));
-	return false;
-}
-
 FString FGenericPlatformProcess::GetApplicationName( uint32 ProcessId )
 {
 	UE_LOG(LogHAL, Fatal, TEXT("FGenericPlatformProcess::GetApplicationName not implemented on this platform"));
@@ -611,4 +605,4 @@ void FGenericPlatformProcess::TearDown()
 {
 	TLazySingleton<FEventPool<EEventPoolTypes::AutoReset>>::TearDown();
 	TLazySingleton<FEventPool<EEventPoolTypes::ManualReset>>::TearDown();
-}
\ No newline at end of file
+}
diff --git a/Engine/Source/Runtime/Core/Private/Mac/MacPlatformProcess.cpp b/Engine/Source/Runtime/Core/Private/Mac/MacPlatformProcess.cpp
--- a/Engine/Source/Runtime/Core/Private/Mac/MacPlatformProcess.cpp
+++ b/Engine/Source/Runtime/Core/Private/Mac/MacPlatformProcess.cpp
@@ -675,7 +675,10 @@ bool FMacPlatformProcess::GetProcReturnCode( FProcHandle& ProcessHandle, int32*
 		return false;
 	}
 
-	*ReturnCode = [(NSTask*)ProcessHandle.Get() terminationStatus];
+	if (ReturnCode)
+	{
+		*ReturnCode = [(NSTask*)ProcessHandle.Get() terminationStatus];
+	}
 	return true;
 }
 
diff --git a/Engine/Source/Runtime/Core/Private/Windows/WindowsPlatformProcess.cpp b/Engine/Source/Runtime/Core/Private/Windows/WindowsPlatformProcess.cpp
--- a/Engine/Source/Runtime/Core/Private/Windows/WindowsPlatformProcess.cpp
+++ b/Engine/Source/Runtime/Core/Private/Windows/WindowsPlatformProcess.cpp
@@ -31,15 +31,6 @@
 	#include <LM.h>
 	#include <Psapi.h>
 	#include <TlHelp32.h>
-
-	namespace ProcessConstants
-	{
-		uint32 WIN_STD_INPUT_HANDLE = STD_INPUT_HANDLE;
-		uint32 WIN_STD_OUTPUT_HANDLE = STD_OUTPUT_HANDLE;
-		uint32 WIN_ATTACH_PARENT_PROCESS = ATTACH_PARENT_PROCESS;		
-		uint32 WIN_STILL_ACTIVE = STILL_ACTIVE;
-	}
-
 #include "Windows/HideWindowsPlatformTypes.h"
 #include "Windows/WindowsPlatformMisc.h"
 
@@ -471,7 +462,16 @@ void FWindowsPlatformProcess::SetThreadAffinityMask( uint64 AffinityMask )
 
 bool FWindowsPlatformProcess::GetProcReturnCode( FProcHandle & ProcHandle, int32* ReturnCode )
 {
-	return ::GetExitCodeProcess( ProcHandle.Get(), (::DWORD *)ReturnCode ) && *((uint32*)ReturnCode) != ProcessConstants::WIN_STILL_ACTIVE;
+	DWORD ExitCode = 0;
+	if (::GetExitCodeProcess(ProcHandle.Get(), &ExitCode) && ExitCode != STILL_ACTIVE)
+	{
+		if (ReturnCode)
+		{
+			*ReturnCode = (int32)ExitCode;
+		}
+		return true;
+	}
+	return false;
 }
 
 bool FWindowsPlatformProcess::GetApplicationMemoryUsage(uint32 ProcessId, SIZE_T* OutMemoryUsage)
@@ -642,15 +642,6 @@ bool FWindowsPlatformProcess::IsApplicationRunning( const TCHAR* ProcName )
 	return false;
 }
 
-bool FWindowsPlatformProcess::IsApplicationAlive(uint32 ProcessId)
-{
-	DWORD ExitCode;
-	HANDLE ProcessHandle = ::OpenProcess(PROCESS_QUERY_INFORMATION, false, ProcessId);
-	const BOOL result = GetExitCodeProcess(ProcessHandle, &ExitCode);
-	CloseHandle(ProcessHandle);
-	return result && ExitCode == STILL_ACTIVE;
-}
-
 FString FWindowsPlatformProcess::GetApplicationName( uint32 ProcessId )
 {
 	FString Output = TEXT("");
diff --git a/Engine/Source/Runtime/Core/Public/GenericPlatform/GenericPlatformProcess.h b/Engine/Source/Runtime/Core/Public/GenericPlatform/GenericPlatformProcess.h
--- a/Engine/Source/Runtime/Core/Public/GenericPlatform/GenericPlatformProcess.h
+++ b/Engine/Source/Runtime/Core/Public/GenericPlatform/GenericPlatformProcess.h
@@ -412,10 +412,6 @@ struct CORE_API FGenericPlatformProcess
 	/** Returns true if the specified application is running */
 	static bool IsApplicationRunning( const TCHAR* ProcName );
 
-	/** Returns true if the specified application has alive, false if it has exited. This is slightly different that IsApplicationRunning,
-	as we are able to query applications that has crashed but not yet exited. */
-	static bool IsApplicationAlive(uint32 ProcessId);
-
 	/** Returns the Name of process given by the PID.  Returns Empty string "" if PID not found. */
 	static FString GetApplicationName( uint32 ProcessId );
 
diff --git a/Engine/Source/Runtime/Core/Public/Windows/WindowsPlatformProcess.h b/Engine/Source/Runtime/Core/Public/Windows/WindowsPlatformProcess.h
--- a/Engine/Source/Runtime/Core/Public/Windows/WindowsPlatformProcess.h
+++ b/Engine/Source/Runtime/Core/Public/Windows/WindowsPlatformProcess.h
@@ -163,7 +163,6 @@ public:
 	static bool GetPerFrameProcessorUsage(uint32 ProcessId, float& ProcessUsageFraction, float& IdleUsageFraction);
 	static bool IsApplicationRunning( uint32 ProcessId );
 	static bool IsApplicationRunning( const TCHAR* ProcName );
-	static bool IsApplicationAlive(uint32 ProcessId);
 	static FString GetApplicationName( uint32 ProcessId );	
 	static bool ExecProcess(const TCHAR* URL, const TCHAR* Params, int32* OutReturnCode, FString* OutStdOut, FString* OutStdErr, const TCHAR* OptionalWorkingDirectory = NULL);
 	static bool ExecElevatedProcess(const TCHAR* URL, const TCHAR* Params, int32* OutReturnCode);
diff --git a/Engine/Source/Runtime/Engine/Private/EngineAnalytics.cpp b/Engine/Source/Runtime/Engine/Private/EngineAnalytics.cpp
--- a/Engine/Source/Runtime/Engine/Private/EngineAnalytics.cpp
+++ b/Engine/Source/Runtime/Engine/Private/EngineAnalytics.cpp
@@ -176,7 +176,7 @@ void FEngineAnalytics::Initialize()
 			// if we're using out-of-process crash reporting, then we don't need to create a sender in this process.
 			if (!FGenericCrashContext::IsOutOfProcessCrashReporter())
 			{
-				SessionSummarySender = MakeShareable(new FEditorSessionSummarySender(FEngineAnalytics::GetProvider(), TEXT("Editor")));
+				SessionSummarySender = MakeShared<FEditorSessionSummarySender>(FEngineAnalytics::GetProvider(), TEXT("Editor"), FPlatformProcess::GetCurrentProcessId());
 			}
 		}
 #endif
