Added missing CloseProc call

#rnx

diff --git a/Engine/Source/Programs/CrashReportClient/Private/CrashReportClientApp.cpp b/Engine/Source/Programs/CrashReportClient/Private/CrashReportClientApp.cpp
--- a/Engine/Source/Programs/CrashReportClient/Private/CrashReportClientApp.cpp
+++ b/Engine/Source/Programs/CrashReportClient/Private/CrashReportClientApp.cpp
@@ -718,6 +718,8 @@ void RunCrashReportClient(const TCHAR* CommandLine)
 		}
 #endif
 
+		FPlatformProcess::CloseProc(MonitoredProcess);
+
 		FCrashReportAnalytics::Shutdown();
 	}
 
