- Fixed CrashReportApp not reading the return code once the monitored process exited (Editor) because the logical && would early out and prevent reading it (keeping it 0).

Edigrated CL 11445717

diff --git a/Engine/Source/Programs/CrashReportClient/Private/CrashReportClientApp.cpp b/Engine/Source/Programs/CrashReportClient/Private/CrashReportClientApp.cpp
--- a/Engine/Source/Programs/CrashReportClient/Private/CrashReportClientApp.cpp
+++ b/Engine/Source/Programs/CrashReportClient/Private/CrashReportClientApp.cpp
@@ -788,16 +788,25 @@ void RunCrashReportClient(const TCHAR* CommandLine)
 			UE_LOG(CrashReportClientLog, Error, TEXT("Failed to open monitor process handle!"));
 		}
 
-		auto IsMonitoredProcessAlive = [&MonitoredProcess](int32& OutReturnCode) -> bool
+		auto GetProcessStatus = [](FProcHandle& ProcessHandle) -> TTuple<bool/*Running*/, int32/*ReturnCode*/>
 		{
-			// The monitored process is considered "alive" if the process is running, and no return code has been set
-			return MonitoredProcess.IsValid() && FPlatformProcess::IsProcRunning(MonitoredProcess) && !FPlatformProcess::GetProcReturnCode(MonitoredProcess, &OutReturnCode);
+			bool bRunning = true;
+			int32 ProcessReturnCode = 0;
+			if (!ProcessHandle.IsValid())
+			{
+				bRunning = false;
+				ProcessReturnCode = -1; // The monitored process crashed before this process could poke at it?
+			}
+			else if (FPlatformProcess::GetProcReturnCode(ProcessHandle, &ProcessReturnCode)) // Return code is available?
+			{
+				bRunning = false; // Not running anymore, the return code was set above.
+			}
+			return MakeTuple(bRunning, ProcessReturnCode);
 		};
 
-		// This IsApplicationAlive() call is quite expensive, perform it at low frequency.
-		int32 ApplicationReturnCode = 0;
-		bool bApplicationAlive = IsMonitoredProcessAlive(ApplicationReturnCode);
-		while (bApplicationAlive && !IsEngineExitRequested())
+		// This GetProcessStatus() call is expensive, perform it at low frequency.
+		TTuple<bool/*bRunning*/, int32/*ExitCode*/> ProcessStatus = GetProcessStatus(MonitoredProcess);
+		while (ProcessStatus.Get<0>() && !IsEngineExitRequested())
 		{
 			const double CurrentTime = FPlatformTime::Seconds();
 
@@ -856,7 +865,7 @@ void RunCrashReportClient(const TCHAR* CommandLine)
 			// Check if the application is alive about every second. (This is an expensive call)
 			if (GFrameCounter % IdealFramerate == 0)
 			{
-				bApplicationAlive = IsMonitoredProcessAlive(ApplicationReturnCode);
+				ProcessStatus = GetProcessStatus(MonitoredProcess);
 			}
 
 			LastTime = CurrentTime;
@@ -874,8 +883,8 @@ void RunCrashReportClient(const TCHAR* CommandLine)
 				if (TempCrashContext.UserSettings.bSendUsageData)
 				{
 					// Query this again, as the crash reporting loop above may have exited before setting this information (via IsEngineExitRequested)
-					bApplicationAlive = IsMonitoredProcessAlive(ApplicationReturnCode);
-					if (!bApplicationAlive)
+					ProcessStatus = GetProcessStatus(MonitoredProcess);
+					if (!ProcessStatus.Get<0>()) // Process is 'not running' anymore?
 					{
 						FCrashReportAnalytics::Initialize();
 
@@ -883,7 +892,7 @@ void RunCrashReportClient(const TCHAR* CommandLine)
 						{
 							// initialize analytics
 							TUniquePtr<FEditorSessionSummarySender> EditorSessionSummarySender = MakeUnique<FEditorSessionSummarySender>(FCrashReportAnalytics::GetProvider(), TEXT("CrashReportClient"), MonitorPid);
-							EditorSessionSummarySender->SetCurrentSessionExitCode(MonitorPid, ApplicationReturnCode);
+							EditorSessionSummarySender->SetCurrentSessionExitCode(MonitorPid, ProcessStatus.Get<1>());
 
 							if (TempCrashContext.UserSettings.bSendUnattendedBugReports)
 							{
