Crash report client doesn't need full access handle to runtime when monitoring.

While monitoring the parent process CRC doesn't need a full access process handle on Windows. Open the handle using limited acccess flags instead.

diff --git a/Engine/Source/Programs/CrashReportClient/Private/CrashReportClientApp.cpp b/Engine/Source/Programs/CrashReportClient/Private/CrashReportClientApp.cpp
--- a/Engine/Source/Programs/CrashReportClient/Private/CrashReportClientApp.cpp
+++ b/Engine/Source/Programs/CrashReportClient/Private/CrashReportClientApp.cpp
@@ -781,8 +781,12 @@ void RunCrashReportClient(const TCHAR* CommandLine)
 		// Starts the disaster recovery service. This records transactions and allows users to recover from previous crashes.
 		RecoveryServicePtr = MakeShared<FRecoveryService>(MonitorPid);
 #endif
-
+#if PLATFORM_WINDOWS
+		// We do not need to open a full access process handle when monitoring process health.
+		FProcHandle MonitoredProcess = FProcHandle(::OpenProcess(PROCESS_DUP_HANDLE | PROCESS_QUERY_LIMITED_INFORMATION | PROCESS_TERMINATE | SYNCHRONIZE, 0, MonitorPid));
+#else
 		FProcHandle MonitoredProcess = FPlatformProcess::OpenProcess(MonitorPid);
+#endif
 		if (!MonitoredProcess.IsValid())
 		{
 			UE_LOG(CrashReportClientLog, Error, TEXT("Failed to open monitor process handle!"));
