- Prevented CrashReportClient::FinalizeDiagnoseReportWorker() function to called after the CrashReportClient instance was deleted.

Details:

The FDiagnoseReportWorker asynchronous task created an extra asynchronous tasks (calling back CrashReportClient instance) that could fire after the targetted CrashReportClient instance was deleted. The solutions moves the logic to 'finilize' the report in the tick function that is expected to run in the game thread and make 'Close Without Sending' flows  like 'Send and Close'/'Send and Restart' but without sending anything.

diff --git a/Engine/Source/Programs/CrashReportClient/Private/CrashReportClient.cpp b/Engine/Source/Programs/CrashReportClient/Private/CrashReportClient.cpp
--- a/Engine/Source/Programs/CrashReportClient/Private/CrashReportClient.cpp
+++ b/Engine/Source/Programs/CrashReportClient/Private/CrashReportClient.cpp
@@ -89,7 +89,9 @@ void FCrashReportClient::StopBackgroundThread()
 
 FReply FCrashReportClient::CloseWithoutSending()
 {
-	RequestEngineExit(TEXT("FCrashReportClient::CloseWithoutSending()"));
+	bSendData = false;
+	bShouldWindowBeHidden = true;
+	StartTicker();
 	return FReply::Handled();
 }
 
@@ -239,7 +241,15 @@ bool FCrashReportClient::Tick(float UnusedDeltaTime)
 	{
 		return true;
 	}
-	
+
+	if (DiagnoseReportTask)
+	{
+		check(DiagnoseReportTask->IsWorkDone()); // Expected when IsProcessingCallstack() returns false.
+		StopBackgroundThread(); // Free the DiagnoseReportTask
+		FinalizeDiagnoseReportWorker(); // Finalization is done on the game thread to avoid race conditions on DiagnosticText and FormattedDiagnosticText.
+		check(DiagnoseReportTask == nullptr); // Expected after StopBackgroundThread() call.
+	}
+
 	if( bSendData )
 	{
 		if (!FCrashUploadBase::IsInitialized())
@@ -315,7 +325,6 @@ void FCrashReportClient::FinalizeDiagnoseReportWorker()
 	FormattedDiagnosticText = FCrashReportUtil::FormatDiagnosticText( DiagnosticText );
 }
 
-
 bool FCrashReportClient::IsProcessingCallstack() const
 {
 	return DiagnoseReportTask && !DiagnoseReportTask->IsWorkDone();
@@ -328,13 +337,6 @@ FDiagnoseReportWorker::FDiagnoseReportWorker( FCrashReportClient* InCrashReportC
 void FDiagnoseReportWorker::DoWork()
 {
 	CrashReportClient->ErrorReport.DiagnoseReport();
-
-	// Inform the game thread that we are done.
-	FSimpleDelegateGraphTask::CreateAndDispatchWhenReady
-	(
-		FSimpleDelegateGraphTask::FDelegate::CreateRaw( CrashReportClient, &FCrashReportClient::FinalizeDiagnoseReportWorker ),
-		TStatId(), nullptr, ENamedThreads::GameThread
-	);
 }
 
 #endif // !CRASH_REPORT_UNATTENDED_ONLY
