- Added a 60 seconds grace period for the Editor process to exit so that we can read its exit code.

diff --git a/Engine/Source/Programs/CrashReportClient/Private/CrashReportClientApp.cpp b/Engine/Source/Programs/CrashReportClient/Private/CrashReportClientApp.cpp
--- a/Engine/Source/Programs/CrashReportClient/Private/CrashReportClientApp.cpp
+++ b/Engine/Source/Programs/CrashReportClient/Private/CrashReportClientApp.cpp
@@ -788,16 +788,30 @@ void RunCrashReportClient(const TCHAR* CommandLine)
 			UE_LOG(CrashReportClientLog, Error, TEXT("Failed to open monitor process handle!"));
 		}
 
-		auto IsMonitoredProcessAlive = [&MonitoredProcess](int32& OutReturnCode) -> bool
+		auto GetProcessStatus = [](FProcHandle& ProcessHandle) -> TTuple<bool/*Running*/, TOptional<int32>/*ReturnCode*/>
 		{
-			// The monitored process is considered "alive" if the process is running, and no return code has been set
-			return MonitoredProcess.IsValid() && FPlatformProcess::IsProcRunning(MonitoredProcess) && !FPlatformProcess::GetProcReturnCode(MonitoredProcess, &OutReturnCode);
+			bool bRunning = true;
+			TOptional<int32> ProcessReturnCodeOpt; // Unknown by default.
+			if (!ProcessHandle.IsValid())
+			{
+				bRunning = false;
+			}
+			else if (!FPlatformProcess::IsProcRunning(ProcessHandle))
+			{
+				bRunning = false;
+				int32 ProcessReturnCode = 0;
+				if (FPlatformProcess::GetProcReturnCode(ProcessHandle, &ProcessReturnCode)) // Is the return code available? (Not supported on all platforms)
+				{
+					ProcessReturnCodeOpt.Emplace(ProcessReturnCode);
+				}
+			}
+
+			return MakeTuple(bRunning, ProcessReturnCodeOpt);
 		};
 
-		// This IsApplicationAlive() call is quite expensive, perform it at low frequency.
-		int32 ApplicationReturnCode = 0;
-		bool bApplicationAlive = IsMonitoredProcessAlive(ApplicationReturnCode);
-		while (bApplicationAlive && !IsEngineExitRequested())
+		// This GetProcessStatus() call is expensive, perform it at low frequency.
+		TTuple<bool/*bRunning*/, TOptional<int32>/*ExitCode*/> ProcessStatus = GetProcessStatus(MonitoredProcess);
+		while (ProcessStatus.Get<0>() && !IsEngineExitRequested())
 		{
 			const double CurrentTime = FPlatformTime::Seconds();
 
@@ -816,7 +830,7 @@ void RunCrashReportClient(const TCHAR* CommandLine)
 
 					// Build error report in memory.
 					FPlatformErrorReport ErrorReport = CollectErrorReport(RecoveryServicePtr.Get(), MonitorPid, CrashContext, MonitorWritePipe);
-					
+
 #if CRASH_REPORT_WITH_RECOVERY
 					if (RecoveryServicePtr && !FPrimaryCrashProperties::Get()->bIsEnsure)
 					{
@@ -856,7 +870,7 @@ void RunCrashReportClient(const TCHAR* CommandLine)
 			// Check if the application is alive about every second. (This is an expensive call)
 			if (GFrameCounter % IdealFramerate == 0)
 			{
-				bApplicationAlive = IsMonitoredProcessAlive(ApplicationReturnCode);
+				ProcessStatus = GetProcessStatus(MonitoredProcess);
 			}
 
 			LastTime = CurrentTime;
@@ -864,6 +878,17 @@ void RunCrashReportClient(const TCHAR* CommandLine)
 
 #if CRASH_REPORT_WITH_MTBF
 		{
+			// The loop above can exit before the Editor (monitored process) exits (because of IsEngineExitRequested()) if the user clicks 'Close Without Sending' very quickly, but for MTBF,
+			// it is desirable to have the Editor process return code. Give some extra time to the Editor to exit. If it doesn't exit within 60 seconds the next CRC instance will sent the
+			// current analytic report delayed, not ideal, but supported.
+			float WaitTimeSeconds = 0;
+			while (ProcessStatus.Get<0>() && WaitTimeSeconds < 60.f)
+			{
+				FPlatformProcess::Sleep(0.1f);
+				WaitTimeSeconds += 0.1f;
+				ProcessStatus = GetProcessStatus(MonitoredProcess);
+			}
+
 			// load our temporary crash context file
 			FSharedCrashContext TempCrashContext;
 			FMemory::Memzero(TempCrashContext);
@@ -874,8 +899,8 @@ void RunCrashReportClient(const TCHAR* CommandLine)
 				if (TempCrashContext.UserSettings.bSendUsageData)
 				{
 					// Query this again, as the crash reporting loop above may have exited before setting this information (via IsEngineExitRequested)
-					bApplicationAlive = IsMonitoredProcessAlive(ApplicationReturnCode);
-					if (!bApplicationAlive)
+					ProcessStatus = GetProcessStatus(MonitoredProcess);
+					if (!ProcessStatus.Get<0>()) // Process is 'not running' anymore?
 					{
 						FCrashReportAnalytics::Initialize();
 
@@ -883,7 +908,11 @@ void RunCrashReportClient(const TCHAR* CommandLine)
 						{
 							// initialize analytics
 							TUniquePtr<FEditorSessionSummarySender> EditorSessionSummarySender = MakeUnique<FEditorSessionSummarySender>(FCrashReportAnalytics::GetProvider(), TEXT("CrashReportClient"), MonitorPid);
-							EditorSessionSummarySender->SetCurrentSessionExitCode(MonitorPid, ApplicationReturnCode);
+							TOptional<int32> ExitCodeOpt = ProcessStatus.Get<1>();
+							if (ExitCodeOpt.IsSet()) // Is it known?
+							{
+								EditorSessionSummarySender->SetCurrentSessionExitCode(MonitorPid, ExitCodeOpt.GetValue());
+							}
 
 							if (TempCrashContext.UserSettings.bSendUnattendedBugReports)
 							{
