- Added a 60 seconds grace period for the Editor process to exit so that we can read its exit code.

diff --git a/Engine/Source/Programs/CrashReportClient/Private/CrashReportClientApp.cpp b/Engine/Source/Programs/CrashReportClient/Private/CrashReportClientApp.cpp
--- a/Engine/Source/Programs/CrashReportClient/Private/CrashReportClientApp.cpp
+++ b/Engine/Source/Programs/CrashReportClient/Private/CrashReportClientApp.cpp
@@ -792,24 +792,29 @@ void RunCrashReportClient(const TCHAR* CommandLine)
 			UE_LOG(CrashReportClientLog, Error, TEXT("Failed to open monitor process handle!"));
 		}
 
-		auto GetProcessStatus = [](FProcHandle& ProcessHandle) -> TTuple<bool/*Running*/, int32/*ReturnCode*/>
+		auto GetProcessStatus = [](FProcHandle& ProcessHandle) -> TTuple<bool/*Running*/, TOptional<int32>/*ReturnCode*/>
 		{
 			bool bRunning = true;
-			int32 ProcessReturnCode = 0;
+			TOptional<int32> ProcessReturnCodeOpt; // Unknown by default.
 			if (!ProcessHandle.IsValid())
 			{
 				bRunning = false;
-				ProcessReturnCode = -1; // The monitored process crashed before this process could poke at it?
 			}
-			else if (FPlatformProcess::GetProcReturnCode(ProcessHandle, &ProcessReturnCode)) // Return code is available?
+			else if (!FPlatformProcess::IsProcRunning(ProcessHandle))
 			{
-				bRunning = false; // Not running anymore, the return code was set above.
+				bRunning = false;
+				int32 ProcessReturnCode = 0;
+				if (FPlatformProcess::GetProcReturnCode(ProcessHandle, &ProcessReturnCode)) // Is the return code available? (Not supported on all platforms)
+				{
+					ProcessReturnCodeOpt.Emplace(ProcessReturnCode);
+				}
 			}
-			return MakeTuple(bRunning, ProcessReturnCode);
+
+			return MakeTuple(bRunning, ProcessReturnCodeOpt);
 		};
 
 		// This GetProcessStatus() call is expensive, perform it at low frequency.
-		TTuple<bool/*bRunning*/, int32/*ExitCode*/> ProcessStatus = GetProcessStatus(MonitoredProcess);
+		TTuple<bool/*bRunning*/, TOptional<int32>/*ExitCode*/> ProcessStatus = GetProcessStatus(MonitoredProcess);
 		while (ProcessStatus.Get<0>() && !IsEngineExitRequested())
 		{
 			const double CurrentTime = FPlatformTime::Seconds();
@@ -829,7 +834,7 @@ void RunCrashReportClient(const TCHAR* CommandLine)
 
 					// Build error report in memory.
 					FPlatformErrorReport ErrorReport = CollectErrorReport(RecoveryServicePtr.Get(), MonitorPid, CrashContext, MonitorWritePipe);
-					
+
 #if CRASH_REPORT_WITH_RECOVERY
 					if (RecoveryServicePtr && !FPrimaryCrashProperties::Get()->bIsEnsure)
 					{
@@ -877,6 +882,17 @@ void RunCrashReportClient(const TCHAR* CommandLine)
 
 #if CRASH_REPORT_WITH_MTBF
 		{
+			// The loop above can exit before the Editor (monitored process) exits (because of IsEngineExitRequested()) if the user clicks 'Close Without Sending' very quickly, but for MTBF,
+			// it is desirable to have the Editor process return code. Give some extra time to the Editor to exit. If it doesn't exit within 60 seconds the next CRC instance will sent the
+			// current analytic report delayed, not ideal, but supported.
+			float WaitTimeSeconds = 0;
+			while (ProcessStatus.Get<0>() && WaitTimeSeconds < 60.f)
+			{
+				FPlatformProcess::Sleep(0.1f);
+				WaitTimeSeconds += 0.1f;
+				ProcessStatus = GetProcessStatus(MonitoredProcess);
+			}
+
 			// load our temporary crash context file
 			FSharedCrashContext TempCrashContext;
 			FMemory::Memzero(TempCrashContext);
@@ -896,7 +912,11 @@ void RunCrashReportClient(const TCHAR* CommandLine)
 						{
 							// initialize analytics
 							TUniquePtr<FEditorSessionSummarySender> EditorSessionSummarySender = MakeUnique<FEditorSessionSummarySender>(FCrashReportAnalytics::GetProvider(), TEXT("CrashReportClient"), MonitorPid);
-							EditorSessionSummarySender->SetCurrentSessionExitCode(MonitorPid, ProcessStatus.Get<1>());
+							TOptional<int32> ExitCodeOpt = ProcessStatus.Get<1>();
+							if (ExitCodeOpt.IsSet()) // Is it known?
+							{
+								EditorSessionSummarySender->SetCurrentSessionExitCode(MonitorPid, ExitCodeOpt.GetValue());
+							}
 
 							if (TempCrashContext.UserSettings.bSendUnattendedBugReports)
 							{
