- Changed the execution flow to ensure the callstack was displayed while the windows was still on screen (it was updated just after the window was closed).

diff --git a/Engine/Source/Programs/CrashReportClient/Private/CrashReportClient.cpp b/Engine/Source/Programs/CrashReportClient/Private/CrashReportClient.cpp
--- a/Engine/Source/Programs/CrashReportClient/Private/CrashReportClient.cpp
+++ b/Engine/Source/Programs/CrashReportClient/Private/CrashReportClient.cpp
@@ -53,6 +53,7 @@ FCrashReportClient::FCrashReportClient(const FPlatformErrorReport& InErrorReport
 			{
 				DiagnoseReportTask = new FAsyncTask<FDiagnoseReportWorker>( this );
 				DiagnoseReportTask->StartBackgroundTask();
+				StartTicker();
 			}
 			else
 			{
@@ -71,9 +72,13 @@ FCrashReportClient::FCrashReportClient(const FPlatformErrorReport& InErrorReport
 	}
 }
 
-
 FCrashReportClient::~FCrashReportClient()
 {
+	if (TickHandle.IsValid())
+	{
+		FTicker::GetCoreTicker().RemoveTicker(TickHandle);
+		TickHandle.Reset();
+	}
 	StopBackgroundThread();
 }
 
@@ -222,7 +227,10 @@ void FCrashReportClient::SendLogFile_OnCheckStateChanged( ECheckBoxState NewRadi
 
 void FCrashReportClient::StartTicker()
 {
-	FTicker::GetCoreTicker().AddTicker(FTickerDelegate::CreateSP(this, &FCrashReportClient::Tick), 1.f);
+	if (!TickHandle.IsValid())
+	{
+		TickHandle = FTicker::GetCoreTicker().AddTicker(FTickerDelegate::CreateRaw(this, &FCrashReportClient::Tick), 1.f);
+	}
 }
 
 void FCrashReportClient::StoreCommentAndUpload()
@@ -245,11 +253,17 @@ bool FCrashReportClient::Tick(float UnusedDeltaTime)
 	if (DiagnoseReportTask)
 	{
 		check(DiagnoseReportTask->IsWorkDone()); // Expected when IsProcessingCallstack() returns false.
-		StopBackgroundThread(); // Free the DiagnoseReportTask
-		FinalizeDiagnoseReportWorker(); // Finalization is done on the game thread to avoid race conditions on DiagnosticText and FormattedDiagnosticText.
+		StopBackgroundThread(); // Free the DiagnoseReportTask to avoid reentering this condition.
+		FinalizeDiagnoseReportWorker(); // Update the Text displaying call stack information (on game thread as they are visualized in UI)
 		check(DiagnoseReportTask == nullptr); // Expected after StopBackgroundThread() call.
 	}
 
+	// Before going further, wait for the an action, either Submit(), CloseWithoutSending() or RequestCloseWindow().
+	if (!bShouldWindowBeHidden)
+	{
+		return true;
+	}
+
 	if( bSendData )
 	{
 		if (!FCrashUploadBase::IsInitialized())
diff --git a/Engine/Source/Programs/CrashReportClient/Private/CrashReportClient.h b/Engine/Source/Programs/CrashReportClient/Private/CrashReportClient.h
--- a/Engine/Source/Programs/CrashReportClient/Private/CrashReportClient.h
+++ b/Engine/Source/Programs/CrashReportClient/Private/CrashReportClient.h
@@ -194,6 +194,9 @@ private:
 	/** Is the uploading complete. */
 	bool bIsUploadComplete;
 
+	/** To know if the ticker was started.*/
+	FDelegateHandle TickHandle;
+
 };
 
 #endif // !CRASH_REPORT_UNATTENDED_ONLY
