CrashReporterClientEditor crash if the bug report windows in closed before the callstack appears
- Fixed CrashReportClient using the analytics provider without checking if it is available first.

diff --git a/Engine/Source/Programs/CrashReportClient/Private/CrashReportClientApp.cpp b/Engine/Source/Programs/CrashReportClient/Private/CrashReportClientApp.cpp
--- a/Engine/Source/Programs/CrashReportClient/Private/CrashReportClientApp.cpp
+++ b/Engine/Source/Programs/CrashReportClient/Private/CrashReportClientApp.cpp
@@ -848,8 +848,11 @@ void RunCrashReportClient(const TCHAR* CommandLine)
 #endif
 					if (bReportCrashAnalyticInfo)
 					{
-						// If analytics is enabled make sure they are submitted now.
-						FCrashReportAnalytics::GetProvider().BlockUntilFlushed(5.0f);
+						if (FCrashReportAnalytics::IsAvailable())
+						{
+							// If analytics is enabled make sure they are submitted now.
+							FCrashReportAnalytics::GetProvider().BlockUntilFlushed(5.0f);
+						}
 						FCrashReportAnalytics::Shutdown();
 					}
 				}
