Crash reporter was stripping callstack frames from the wrong direction.

When creating portable callstack for the crash reports a number of stack frames should be skipped, corresponding to the crash collecting code itself. This value was applied to wrong end of the callstack.

diff --git a/Engine/Source/Programs/CrashReportClient/Private/CrashReportClientApp.cpp b/Engine/Source/Programs/CrashReportClient/Private/CrashReportClientApp.cpp
--- a/Engine/Source/Programs/CrashReportClient/Private/CrashReportClientApp.cpp
+++ b/Engine/Source/Programs/CrashReportClient/Private/CrashReportClientApp.cpp
@@ -441,7 +441,7 @@ FPlatformErrorReport CollectErrorReport(FRecoveryService* RecoveryService, uint3
 		if (ThreadId == SharedCrashContext.CrashingThreadId)
 		{
 			CrashContext.SetPortableCallStack(
-				StackFrames + SharedCrashContext.NumStackFramesToIgnore,
+				StackFrames,
 				StackFrameCount - SharedCrashContext.NumStackFramesToIgnore
 			);
 		}
