Tried to get more reliable Editor analytic data out of crash rerporter app:
  - Ensured the data is fully written and read to/form the pipe between the Editor and CrashReporter.
  - Refactored the code sending MTFB to send it as soon as possible
  - Added some special exit code to detect when the Editor is still running or the exit code is unknown when the summary event is sent.

diff --git a/Engine/Source/Programs/CrashReportClient/Private/CrashReportClientApp.cpp b/Engine/Source/Programs/CrashReportClient/Private/CrashReportClientApp.cpp
--- a/Engine/Source/Programs/CrashReportClient/Private/CrashReportClientApp.cpp
+++ b/Engine/Source/Programs/CrashReportClient/Private/CrashReportClientApp.cpp
@@ -553,14 +553,56 @@ SubmitCrashReportResult SendErrorReport(FPlatformErrorReport& ErrorReport,
 
 bool IsCrashReportAvailable(uint32 WatchedProcess, FSharedCrashContext& CrashContext, void* ReadPipe)
 {
-	static TArray<uint8> Buffer;
-	Buffer.Reserve(8 * 1024); // This allocates only once because Buffer is static.
+	TArray<uint8> Buffer;
 
+	// Is data available on the pipe.
 	if (FPlatformProcess::ReadPipeToArray(ReadPipe, Buffer))
 	{
-		FPlatformMemory::Memcpy(&CrashContext, Buffer.GetData(), Buffer.Num());
-		return true;
+		// This is to ensure the FSharedCrashContext compiled in the monitored process and this process has the same size.
+		int32 TotalRead = Buffer.Num();
+
+		// Utility function to copy bytes from a source to a destination buffer.
+		auto CopyFn = [](const TArray<uint8>& SrcData, uint8* DstIt, uint8* DstEndIt)
+		{
+			int32 CopyCount = FMath::Min(SrcData.Num(), static_cast<int32>(DstEndIt - DstIt)); // Limit the number of byte to copy to avoid writing passed the end of the destination.
+			FPlatformMemory::Memcpy(DstIt, SrcData.GetData(), CopyCount);
+			return DstIt + CopyCount; // Returns the updated position.
+		};
+
+		// Iterators to defines the boundaries of the destination buffer in memory.
+		uint8* SharedCtxIt = reinterpret_cast<uint8*>(&CrashContext);
+		uint8* SharedCtxEndIt = SharedCtxIt + sizeof(FSharedCrashContext);
+
+		// Copy the data already read and update the destination iterator.
+		SharedCtxIt = CopyFn(Buffer, SharedCtxIt, SharedCtxEndIt);
+
+		// Try to consume all the expected data within a defined period of time.
+		FDateTime WaitEndTime = FDateTime::UtcNow() + FTimespan(0, 0, 2);
+		while (SharedCtxIt != SharedCtxEndIt && FDateTime::UtcNow() <= WaitEndTime)
+		{
+			if (FPlatformProcess::ReadPipeToArray(ReadPipe, Buffer)) // This is false if no data is available, but the writer may be still be writing.
+			{
+				TotalRead += Buffer.Num();
+				SharedCtxIt = CopyFn(Buffer, SharedCtxIt, SharedCtxEndIt); // Copy the data read.
+			}
+			else
+			{
+				FPlatformProcess::Sleep(0.1); // Give the writer some time.
+			}
+		}
+
+		if (TotalRead < sizeof(FSharedCrashContext))
+		{
+			UE_LOG(CrashReportClientLog, Error, TEXT("The shared crash context emitted by the monitored process could not be fully read (or was smaller than expected by crash reporter)."));
+		}
+		else if (TotalRead > sizeof(FSharedCrashContext))
+		{
+			UE_LOG(CrashReportClientLog, Error, TEXT("The shared crash context emitted by the monitored process is larger than expected by crash reporter."));
+		}
+
+		return SharedCtxIt == SharedCtxEndIt;
 	}
+
 	return false;
 }
 
@@ -882,59 +924,70 @@ void RunCrashReportClient(const TCHAR* CommandLine)
 			LastTime = CurrentTime;
 		}
 
-#if CRASH_REPORT_WITH_MTBF
+#if CRASH_REPORT_WITH_MTBF // Expected to be 1 when compiling CrashReportClientEditor.
 		{
 			// The loop above can exit before the Editor (monitored process) exits (because of IsEngineExitRequested()) if the user clicks 'Close Without Sending' very quickly, but for MTBF,
-			// it is desirable to have the Editor process return code. Give some extra time to the Editor to exit. If it doesn't exit within 60 seconds the next CRC instance will sent the
+			// it is desirable to have the Editor process return code. Give some extra time to the Editor to exit. If it doesn't exit within x seconds the next CRC instance will sent the
 			// current analytic report delayed, not ideal, but supported.
-			float WaitTimeSeconds = 0;
-			while (ProcessStatus.Get<0>() && WaitTimeSeconds < 60.f)
+			FDateTime WaitEndTime = FDateTime::UtcNow() + FTimespan(0, 0, 120);
+			while (ProcessStatus.Get<0>() && FDateTime::UtcNow() <= WaitEndTime)
 			{
-				FPlatformProcess::Sleep(0.1f);
-				WaitTimeSeconds += 0.1f;
+				FPlatformProcess::Sleep(0.5f); // In seconds
 				ProcessStatus = GetProcessStatus(MonitoredProcess);
 			}
 
-			// load our temporary crash context file
-			FSharedCrashContext TempCrashContext;
-			FMemory::Memzero(TempCrashContext);
-			if (LoadTempCrashContextFromFile(TempCrashContext, MonitorPid))
-			{
-				TempCrashContext.SessionContext.ProcessId = MonitorPid;
+			// Editor shutdown state from different point of view.
+			bool bAbnormalShutdownFromEditorPov = false;
+			bool bNormalShutdownFromOsPov = false; // Can only be true if the Editor process exit code is known and equal to zero.
 
-				if (TempCrashContext.UserSettings.bSendUsageData)
+			{
+				// Send the editor summary event(s) first, maximizing chance of sucessfull transmission (less opportunities for bugs to prevent it).
+				FCrashReportAnalytics::Initialize();
+				if (FCrashReportAnalytics::IsAvailable())
 				{
-					// Query this again, as the crash reporting loop above may have exited before setting this information (via IsEngineExitRequested)
+					// NOTE: The Editor doesn't create summary events if analytics are disabled (not permitted to send). It may still send pending events
+					//       accmulated while 'Send Data' was true, but will not send any newer.
+					FEditorSessionSummarySender EditorSessionSummarySender(FCrashReportAnalytics::GetProvider(), TEXT("CrashReportClient"), MonitorPid);
+
+					// Query the Editor process state again, the Editor may still run.
 					ProcessStatus = GetProcessStatus(MonitoredProcess);
 					if (!ProcessStatus.Get<0>()) // Process is 'not running' anymore?
 					{
-						FCrashReportAnalytics::Initialize();
-
-						if (FCrashReportAnalytics::IsAvailable())
+						TOptional<int32> ExitCodeOpt = ProcessStatus.Get<1>();
+						if (ExitCodeOpt.IsSet()) // Exit code is known?
 						{
-							// initialize analytics
-							TUniquePtr<FEditorSessionSummarySender> EditorSessionSummarySender = MakeUnique<FEditorSessionSummarySender>(FCrashReportAnalytics::GetProvider(), TEXT("CrashReportClient"), MonitorPid);
-							TOptional<int32> ExitCodeOpt = ProcessStatus.Get<1>();
-							if (ExitCodeOpt.IsSet()) // Is it known?
-							{
-								EditorSessionSummarySender->SetCurrentSessionExitCode(MonitorPid, ExitCodeOpt.GetValue());
-							}
-
-							if (TempCrashContext.UserSettings.bSendUnattendedBugReports)
-							{
-								// Send a spoofed crash report in the case that we detect an abnormal shutdown has occurred
-								if (WasAbnormalShutdown(*EditorSessionSummarySender.Get()))
-								{
-									HandleAbnormalShutdown(TempCrashContext, MonitorPid, MonitorWritePipe, RecoveryServicePtr);
-								}
-							}
-
-							// send analytics and shutdown
-							EditorSessionSummarySender->Shutdown();
+							EditorSessionSummarySender.SetCurrentSessionExitCode(MonitorPid, ExitCodeOpt.GetValue()); // The Editor exit code from the OS point of view.
+							bNormalShutdownFromOsPov = ExitCodeOpt.GetValue() == 0;
 						}
-
-						FCrashReportAnalytics::Shutdown();
+						else
+						{
+							EditorSessionSummarySender.SetCurrentSessionExitCode(MonitorPid, 112233001); // Special exit code, arbitrary but easy to read in decimal, to mark the process exit code as 'unknown'.
+						}
+					}
+					else
+					{
+						EditorSessionSummarySender.SetCurrentSessionExitCode(MonitorPid, 112233002); // Special exit code, arbitrary, but easy to read in decimal, to mark the process exit code as 'still running'.
 					}
+
+					// Check what the Editor knows about the exit. Was the proper handlers called and the flag(s) set in the summary event?
+					bAbnormalShutdownFromEditorPov = WasAbnormalShutdown(EditorSessionSummarySender);
+
+					// Send summary session event(s).
+					EditorSessionSummarySender.Shutdown();
+				}
+				FCrashReportAnalytics::Shutdown();
+			}
+
+			// If the Editor hasn't called all its crash/exit handlers properly (session summary flags are not set) and the Editor exit code isn't known or different than zero.
+			if (bAbnormalShutdownFromEditorPov && !bNormalShutdownFromOsPov)
+			{
+				// Load our temporary crash context file.
+				FSharedCrashContext TempCrashContext;
+				FMemory::Memzero(TempCrashContext);
+				if (LoadTempCrashContextFromFile(TempCrashContext, MonitorPid) && TempCrashContext.UserSettings.bSendUsageData && TempCrashContext.UserSettings.bSendUnattendedBugReports)
+				{
+					// Send a spoofed crash report in the case that we detect an abnormal shutdown has occurred
+					HandleAbnormalShutdown(TempCrashContext, MonitorPid, MonitorWritePipe, RecoveryServicePtr);
 				}
 			}
 		}
diff --git a/Engine/Source/Runtime/Core/Private/Windows/WindowsPlatformCrashContext.cpp b/Engine/Source/Runtime/Core/Private/Windows/WindowsPlatformCrashContext.cpp
--- a/Engine/Source/Runtime/Core/Private/Windows/WindowsPlatformCrashContext.cpp
+++ b/Engine/Source/Runtime/Core/Private/Windows/WindowsPlatformCrashContext.cpp
@@ -638,25 +638,34 @@ int32 ReportCrashForMonitor(
 	}
 
 	// Write the shared context to the pipe
-	int32 OutDataWritten = 0;
-	FPlatformProcess::WritePipe(WritePipe, (UINT8*)SharedContext, sizeof(FSharedCrashContext), &OutDataWritten);
-	check(OutDataWritten == sizeof(FSharedCrashContext));
-
-	// Wait for a response, saying it's ok to continue
-	bool bCanContinueExecution = false;
-	int32 ExitCode = 0;
-	// Would like to use TInlineAllocator here to avoid heap allocation on crashes, but it doesn't work since ReadPipeToArray 
-	// cannot take array with non-default allocator
-	TArray<uint8> ResponseBuffer;
-	ResponseBuffer.AddZeroed(16);
-	while (!FPlatformProcess::GetProcReturnCode(CrashMonitorHandle, &ExitCode) && !bCanContinueExecution)
+	bool bPipeWriteSucceeded = true;
+	const uint8* DataIt = (const uint8*)SharedContext;
+	const uint8* DataEndIt = DataIt + sizeof(FSharedCrashContext);
+	while (DataIt != DataEndIt && bPipeWriteSucceeded)
 	{
-		if (FPlatformProcess::ReadPipeToArray(ReadPipe, ResponseBuffer))
+		int32 OutDataWritten = 0;
+		bPipeWriteSucceeded = FPlatformProcess::WritePipe(WritePipe, DataIt, static_cast<int32>(DataEndIt - DataIt), &OutDataWritten);
+		DataIt += OutDataWritten;
+	}
+
+	if (bPipeWriteSucceeded) // The receiver is not likely to respond if the shared context wasn't successfully written to the pipe.
+	{
+		// Wait for a response, saying it's ok to continue
+		bool bCanContinueExecution = false;
+		int32 ExitCode = 0;
+		// Would like to use TInlineAllocator here to avoid heap allocation on crashes, but it doesn't work since ReadPipeToArray 
+		// cannot take array with non-default allocator
+		TArray<uint8> ResponseBuffer;
+		ResponseBuffer.AddZeroed(16);
+		while (!FPlatformProcess::GetProcReturnCode(CrashMonitorHandle, &ExitCode) && !bCanContinueExecution)
 		{
-			if (ResponseBuffer[0] == 0xd && ResponseBuffer[1] == 0xe &&
-				ResponseBuffer[2] == 0xa && ResponseBuffer[3] == 0xd)
+			if (FPlatformProcess::ReadPipeToArray(ReadPipe, ResponseBuffer))
 			{
-				bCanContinueExecution = true;
+				if (ResponseBuffer[0] == 0xd && ResponseBuffer[1] == 0xe &&
+					ResponseBuffer[2] == 0xa && ResponseBuffer[3] == 0xd)
+				{
+					bCanContinueExecution = true;
+				}
 			}
 		}
 	}
