Ensure to create a crash report everytime the Editor things this is an abnormal termination even if the Editor exit code is zero as it was in 4.24.3

diff --git a/Engine/Source/Programs/CrashReportClient/Private/CrashReportClientApp.cpp b/Engine/Source/Programs/CrashReportClient/Private/CrashReportClientApp.cpp
--- a/Engine/Source/Programs/CrashReportClient/Private/CrashReportClientApp.cpp
+++ b/Engine/Source/Programs/CrashReportClient/Private/CrashReportClientApp.cpp
@@ -938,7 +938,6 @@ void RunCrashReportClient(const TCHAR* CommandLine)
 
 			// Editor shutdown state from different point of view.
 			bool bAbnormalShutdownFromEditorPov = false;
-			bool bNormalShutdownFromOsPov = false; // Can only be true if the Editor process exit code is known and equal to zero.
 
 			// Send the editor summary event(s) first, maximizing chance of sucessfull transmission (less opportunities for bugs to prevent it).
 			FCrashReportAnalytics::Initialize();
@@ -957,7 +956,6 @@ void RunCrashReportClient(const TCHAR* CommandLine)
 						if (ExitCodeOpt.IsSet()) // Exit code is known?
 						{
 							EditorSessionSummarySender.SetCurrentSessionExitCode(MonitorPid, ExitCodeOpt.GetValue()); // The Editor exit code from the OS point of view.
-							bNormalShutdownFromOsPov = ExitCodeOpt.GetValue() == 0;
 						}
 						else
 						{
@@ -976,8 +974,8 @@ void RunCrashReportClient(const TCHAR* CommandLine)
 					EditorSessionSummarySender.Shutdown();
 				}
 
-				// If the Editor hasn't called all its crash/exit handlers properly (session summary flags are not set) and the Editor exit code isn't known or different than zero.
-				if (bAbnormalShutdownFromEditorPov && !bNormalShutdownFromOsPov)
+				// If the Editor thinks the session ended up abnormally, generate a crash report (to get the Editor logs and figure out why this happened).
+				if (bAbnormalShutdownFromEditorPov)
 				{
 					// Load our temporary crash context file.
 					FSharedCrashContext TempCrashContext;
