- Fixed crash report client editor to prevent sending any usage data.

diff --git a/Engine/Source/Programs/CrashReportClient/Private/CrashReportClientApp.cpp b/Engine/Source/Programs/CrashReportClient/Private/CrashReportClientApp.cpp
--- a/Engine/Source/Programs/CrashReportClient/Private/CrashReportClientApp.cpp
+++ b/Engine/Source/Programs/CrashReportClient/Private/CrashReportClientApp.cpp
@@ -936,22 +936,22 @@ void RunCrashReportClient(const TCHAR* CommandLine)
 				ProcessStatus = GetProcessStatus(MonitoredProcess);
 			}
 
-			// Editor shutdown state from different point of view.
-			bool bAbnormalShutdownFromEditorPov = false;
-
-			// Send the editor summary event(s) first, maximizing chance of sucessfull transmission (less opportunities for bugs to prevent it).
-			FCrashReportAnalytics::Initialize();
-			if (FCrashReportAnalytics::IsAvailable())
+			// Load the temporary crash context file.
+			FSharedCrashContext TempCrashContext;
+			FMemory::Memzero(TempCrashContext);
+			if (!ProcessStatus.Get<0>() && LoadTempCrashContextFromFile(TempCrashContext, MonitorPid) && TempCrashContext.UserSettings.bSendUsageData)
 			{
+				// Send the editor summary event(s) first, maximizing chance of sucessfull transmission (less opportunities for bugs to prevent it).
+				FCrashReportAnalytics::Initialize();
+				if (FCrashReportAnalytics::IsAvailable())
 				{
-					// NOTE: The Editor doesn't create summary events if analytics are disabled (not permitted to send). It may still send pending events
-					//       accmulated while 'Send Data' was true, but will not send any newer.
-					FEditorSessionSummarySender EditorSessionSummarySender(FCrashReportAnalytics::GetProvider(), TEXT("CrashReportClient"), MonitorPid);
-
-					// Query the Editor process state again, the Editor may still run.
-					ProcessStatus = GetProcessStatus(MonitoredProcess);
-					if (!ProcessStatus.Get<0>()) // Process is 'not running' anymore?
+					// Editor shutdown state from Editor point of view.
+					bool bAbnormalShutdownFromEditorPov = false;
 					{
+						FEditorSessionSummarySender EditorSessionSummarySender(FCrashReportAnalytics::GetProvider(), TEXT("CrashReportClient"), MonitorPid);
+
+						// Query the Editor process state again, the Editor may still run.
+						ProcessStatus = GetProcessStatus(MonitoredProcess);
 						TOptional<int32> ExitCodeOpt = ProcessStatus.Get<1>();
 						if (ExitCodeOpt.IsSet()) // Exit code is known?
 						{
@@ -961,33 +961,23 @@ void RunCrashReportClient(const TCHAR* CommandLine)
 						{
 							EditorSessionSummarySender.SetCurrentSessionExitCode(MonitorPid, 112233001); // Special exit code, arbitrary but easy to read in decimal, to mark the process exit code as 'unknown'.
 						}
-					}
-					else
-					{
-						EditorSessionSummarySender.SetCurrentSessionExitCode(MonitorPid, 112233002); // Special exit code, arbitrary, but easy to read in decimal, to mark the process exit code as 'still running'.
-					}
 
-					// Check what the Editor knows about the exit. Was the proper handlers called and the flag(s) set in the summary event?
-					bAbnormalShutdownFromEditorPov = WasAbnormalShutdown(EditorSessionSummarySender);
+						// Check what the Editor knows about the exit. Was the proper handlers called and the flag(s) set in the summary event?
+						bAbnormalShutdownFromEditorPov = WasAbnormalShutdown(EditorSessionSummarySender);
 
-					// Send summary session event(s).
-					EditorSessionSummarySender.Shutdown();
-				}
+						// Send summary session event(s).
+						EditorSessionSummarySender.Shutdown();
+					}
 
-				// If the Editor thinks the session ended up abnormally, generate a crash report (to get the Editor logs and figure out why this happened).
-				if (bAbnormalShutdownFromEditorPov)
-				{
-					// Load our temporary crash context file.
-					FSharedCrashContext TempCrashContext;
-					FMemory::Memzero(TempCrashContext);
-					if (LoadTempCrashContextFromFile(TempCrashContext, MonitorPid) && TempCrashContext.UserSettings.bSendUsageData && TempCrashContext.UserSettings.bSendUnattendedBugReports)
+					// If the Editor thinks the session ended up abnormally, generate a crash report (to get the Editor logs and figure out why this happened).
+					if (bAbnormalShutdownFromEditorPov)
 					{
 						// Send a spoofed crash report in the case that we detect an abnormal shutdown has occurred
 						HandleAbnormalShutdown(TempCrashContext, MonitorPid, MonitorWritePipe, RecoveryServicePtr);
 					}
 				}
+				FCrashReportAnalytics::Shutdown();
 			}
-			FCrashReportAnalytics::Shutdown();
 		}
 #endif
 		// clean up the context file
