Also prevent CrashReportClientEditor from sending abnormal termination logs if bSendUnattendedBugReports is not set.

diff --git a/Engine/Source/Programs/CrashReportClient/Private/CrashReportClientApp.cpp b/Engine/Source/Programs/CrashReportClient/Private/CrashReportClientApp.cpp
--- a/Engine/Source/Programs/CrashReportClient/Private/CrashReportClientApp.cpp
+++ b/Engine/Source/Programs/CrashReportClient/Private/CrashReportClientApp.cpp
@@ -970,7 +970,7 @@ void RunCrashReportClient(const TCHAR* CommandLine)
 					}
 
 					// If the Editor thinks the session ended up abnormally, generate a crash report (to get the Editor logs and figure out why this happened).
-					if (bAbnormalShutdownFromEditorPov)
+					if (bAbnormalShutdownFromEditorPov && TempCrashContext.UserSettings.bSendUnattendedBugReports)
 					{
 						// Send a spoofed crash report in the case that we detect an abnormal shutdown has occurred
 						HandleAbnormalShutdown(TempCrashContext, MonitorPid, MonitorWritePipe, RecoveryServicePtr);
