Prevented one CrashReportClientEditor from sending the report owned by another concurrent instance and losing the exit code.

diff --git a/Engine/Source/Developer/EditorAnalyticsSession/Private/EditorAnalyticsSession.cpp b/Engine/Source/Developer/EditorAnalyticsSession/Private/EditorAnalyticsSession.cpp
--- a/Engine/Source/Developer/EditorAnalyticsSession/Private/EditorAnalyticsSession.cpp
+++ b/Engine/Source/Developer/EditorAnalyticsSession/Private/EditorAnalyticsSession.cpp
@@ -39,6 +39,8 @@ namespace EditorAnalyticsDefs
 	static const FString ProjectVersionStoreKey(TEXT("ProjectVersion"));
 	static const FString EngineVersionStoreKey(TEXT("EngineVersion"));
 	static const FString PlatformProcessIDStoreKey(TEXT("PlatformProcessID"));
+	static const FString MonitorProcessIDStoreKey(TEXT("MonitorProcessID"));
+	static const FString ExitCodeStoreKey(TEXT("ExitCode"));
 
 	// timestamps
 	static const FString StartupTimestampStoreKey(TEXT("StartupTimestamp"));
@@ -245,8 +247,16 @@ namespace EditorAnalyticsUtils
 		GET_STORED_STRING(ProjectDescription);
 		GET_STORED_STRING(ProjectVersion);
 		GET_STORED_STRING(EngineVersion);
-
 		GET_STORED_INT(PlatformProcessID);
+		GET_STORED_INT(MonitorProcessID);
+
+		{
+			FString ExitCodeString;
+			if (FPlatformMisc::GetStoredValue(EditorAnalyticsDefs::StoreId, SectionName, EditorAnalyticsDefs::ExitCodeStoreKey, ExitCodeString))
+			{
+				Session.ExitCode.Emplace(FCString::Atoi(*ExitCodeString));
+			}
+		}
 
 		// scope is just to isolate the temporary value
 		{
@@ -346,6 +356,7 @@ FEditorAnalyticsSession::FEditorAnalyticsSession()
 
 	ProjectName = EditorAnalyticsDefs::UnknownProjectValueString;
 	PlatformProcessID = 0;
+	MonitorProcessID = 0;
 	StartupTimestamp = FDateTime::MinValue();
 	Timestamp = FDateTime::MinValue();
 	Idle1Min = 0;
@@ -427,6 +438,7 @@ bool FEditorAnalyticsSession::Save()
 		TMap<FString, FString> KeyValues = {
 			{EditorAnalyticsDefs::EngineVersionStoreKey,       EngineVersion},
 			{EditorAnalyticsDefs::PlatformProcessIDStoreKey,   FString::FromInt(PlatformProcessID)},
+			{EditorAnalyticsDefs::MonitorProcessIDStoreKey,    FString::FromInt(MonitorProcessID)},
 			{EditorAnalyticsDefs::DesktopGPUAdapterStoreKey,   DesktopGPUAdapter},
 			{EditorAnalyticsDefs::RenderingGPUAdapterStoreKey, RenderingGPUAdapter},
 			{EditorAnalyticsDefs::GPUVendorIDStoreKey,         FString::FromInt(GPUVendorID)},
@@ -476,6 +488,11 @@ bool FEditorAnalyticsSession::Save()
 			{EditorAnalyticsDefs::IsInVRModeStoreKey,     EditorAnalyticsUtils::BoolToStoredString(bIsInVRMode)},
 		};
 
+		if (ExitCode.IsSet())
+		{
+			KeyValues.Emplace(EditorAnalyticsDefs::ExitCodeStoreKey, FString::FromInt(ExitCode.GetValue()));
+		}
+
 		FPlatformMisc::SetStoredValues(EditorAnalyticsDefs::StoreId, StorageLocation, KeyValues);
 	}
 
@@ -514,6 +531,8 @@ bool FEditorAnalyticsSession::Delete() const
 	FPlatformMisc::DeleteStoredValue(EditorAnalyticsDefs::StoreId, SectionName, EditorAnalyticsDefs::ProjectVersionStoreKey);
 	FPlatformMisc::DeleteStoredValue(EditorAnalyticsDefs::StoreId, SectionName, EditorAnalyticsDefs::EngineVersionStoreKey);
 	FPlatformMisc::DeleteStoredValue(EditorAnalyticsDefs::StoreId, SectionName, EditorAnalyticsDefs::PlatformProcessIDStoreKey);
+	FPlatformMisc::DeleteStoredValue(EditorAnalyticsDefs::StoreId, SectionName, EditorAnalyticsDefs::MonitorProcessIDStoreKey);
+	FPlatformMisc::DeleteStoredValue(EditorAnalyticsDefs::StoreId, SectionName, EditorAnalyticsDefs::ExitCodeStoreKey);
 
 	FPlatformMisc::DeleteStoredValue(EditorAnalyticsDefs::StoreId, SectionName, EditorAnalyticsDefs::StartupTimestampStoreKey);
 	FPlatformMisc::DeleteStoredValue(EditorAnalyticsDefs::StoreId, SectionName, EditorAnalyticsDefs::TimestampStoreKey);
@@ -620,3 +639,41 @@ void FEditorAnalyticsSession::LogEvent(EEventType InEventType, const FDateTime&
 {
 	EditorAnalyticsUtils::LogSessionEvent(*this, InEventType, InTimestamp);
 }
+
+bool FEditorAnalyticsSession::FindSession(const uint32 InSessionProcessId, FEditorAnalyticsSession& OutSession)
+{
+	if (!ensure(IsLocked()))
+	{
+		return false;
+	}
+
+	TArray<FString> SessionIDs = EditorAnalyticsUtils::GetSessionList();
+
+	// Retrieve all the sessions in the list from storage
+	for (const FString& Id : SessionIDs)
+	{
+		FEditorAnalyticsSession Session;
+		EditorAnalyticsUtils::LoadInternal(Session, Id);
+		if (Session.PlatformProcessID == InSessionProcessId)
+		{
+			OutSession = MoveTemp(Session);
+			return true;
+		}
+	}
+
+	return false;
+}
+
+void FEditorAnalyticsSession::SaveExitCode(int32 InExitCode)
+{
+	if (!ensure(IsLocked()))
+	{
+		return;
+	}
+
+	ExitCode.Emplace(InExitCode);
+	FString ExitCodeStr = FString::Printf(TFormatSpecifier<decltype(InExitCode)>::GetFormatSpecifier(), InExitCode);
+
+	const FString StorageLocation = EditorAnalyticsUtils::GetSessionStorageLocation(SessionId);
+	FPlatformMisc::SetStoredValue(EditorAnalyticsDefs::StoreId, StorageLocation, EditorAnalyticsDefs::ExitCodeStoreKey, ExitCodeStr);
+}
diff --git a/Engine/Source/Developer/EditorAnalyticsSession/Private/EditorSessionSummarySender.cpp b/Engine/Source/Developer/EditorAnalyticsSession/Private/EditorSessionSummarySender.cpp
--- a/Engine/Source/Developer/EditorAnalyticsSession/Private/EditorSessionSummarySender.cpp
+++ b/Engine/Source/Developer/EditorAnalyticsSession/Private/EditorSessionSummarySender.cpp
@@ -27,7 +27,7 @@ namespace EditorSessionSenderDefs
 	static const FString AbnormalSessionToken(TEXT("AbnormalShutdown"));
 }
 
-FEditorSessionSummarySender::FEditorSessionSummarySender(IAnalyticsProviderET& InAnalyticsProvider, const FString& InSenderName, const int32 InCurrentSessionProcessId)
+FEditorSessionSummarySender::FEditorSessionSummarySender(IAnalyticsProviderET& InAnalyticsProvider, const FString& InSenderName, const uint32 InCurrentSessionProcessId)
 	: HeartbeatTimeElapsed(0.0f)
 	, AnalyticsProvider(InAnalyticsProvider)
 	, Sender(InSenderName)
@@ -56,53 +56,13 @@ void FEditorSessionSummarySender::Shutdown()
 	SendStoredSessions(/*bForceSendCurrentSession*/true);
 }
 
-void FEditorSessionSummarySender::SetCurrentSessionExitCode(const int32 InCurrentSessionProcessId, const int32 InExitCode)
-{
-	check(CurrentSessionProcessId == InCurrentSessionProcessId);
-	CurrentSessionExitCode = InExitCode;
-}
-
-bool FEditorSessionSummarySender::FindCurrentSession(FEditorAnalyticsSession& OutSession) const
-{
-	if (FPlatformProcess::IsApplicationRunning(CurrentSessionProcessId))
-	{
-		// still running, can't be abnormal termination
-		return false;
-	}
-
-	bool bFound = false;
-
-	if (FEditorAnalyticsSession::Lock(FTimespan::FromMilliseconds(100)))
-	{
-		TArray<FEditorAnalyticsSession> ExistingSessions;
-		FEditorAnalyticsSession::LoadAllStoredSessions(ExistingSessions);
-
-		const int32 ProcessID = CurrentSessionProcessId;
-		FEditorAnalyticsSession* CurrentSession = ExistingSessions.FindByPredicate(
-			[ProcessID](const FEditorAnalyticsSession& Session)
-			{
-				return Session.PlatformProcessID == ProcessID;
-			});
-
-		if (CurrentSession != nullptr)
-		{
-			OutSession = *CurrentSession;
-			bFound = true;
-		}
-
-		FEditorAnalyticsSession::Unlock();
-	}
-
-	return bFound;
-}
-
-void FEditorSessionSummarySender::SendStoredSessions(const bool bForceSendCurrentSession) const
+void FEditorSessionSummarySender::SendStoredSessions(const bool bForceSendOwnedSession) const
 {
 	// Load the list of sessions to process. Expect contention on the analytic session lock between the Editor and CrashReportClientEditor (on windows) or between Editor instances (on mac/linux)
 	//   - Try every 'n' seconds if bForceSendCurrentSession is true.
 	//   - Don't block and don't loop if bForceSendCurrentSession is false.
 	TArray<FEditorAnalyticsSession> SessionsToReport;
-	FTimespan Timemout(bForceSendCurrentSession ? FTimespan::FromSeconds(0.5) : FTimespan::Zero());
+	FTimespan Timemout(bForceSendOwnedSession ? FTimespan::FromSeconds(5) : FTimespan::Zero());
 	bool bSessionsLoaded = false;
 	do
 	{
@@ -114,14 +74,25 @@ void FEditorSessionSummarySender::SendStoredSessions(const bool bForceSendCurren
 
 			TArray<FEditorAnalyticsSession> SessionsToDelete;
 
+			// Whether this summary sender instance was tasked to send the specified session.
+			auto HasOwnershipOf = [this](const FEditorAnalyticsSession& InSession) { return InSession.PlatformProcessID == CurrentSessionProcessId; };
+
+			// Whether the process(es) writing and/or sending the specified session are all dead (and cannot write or send the session anymore).
+			auto IsOrphan = [](const FEditorAnalyticsSession& InSession) { return !FPlatformProcess::IsApplicationRunning(InSession.PlatformProcessID) && (InSession.MonitorProcessID == 0 || !FPlatformProcess::IsApplicationRunning(InSession.MonitorProcessID)); };
+
 			// Check each stored session to see if they should be sent or not
 			for (FEditorAnalyticsSession& Session : ExistingSessions)
 			{
-				const bool bForceSendSession = bForceSendCurrentSession && (Session.PlatformProcessID == CurrentSessionProcessId);
-				if (!bForceSendSession && FPlatformProcess::IsApplicationRunning(Session.PlatformProcessID))
+				if (HasOwnershipOf(Session)) // This process was configured to send this session.
 				{
-					// Skip processes that are still running
-					continue;
+					if (!bForceSendOwnedSession) // Don't send until forced to (on shutdown).
+					{
+						continue;
+					}
+				}
+				else if (!IsOrphan(Session))
+				{
+					continue; // Skip, another process is in charge of sending this session.
 				}
 
 				const FTimespan SessionAge = FDateTime::UtcNow() - Session.Timestamp;
@@ -132,6 +103,8 @@ void FEditorSessionSummarySender::SendStoredSessions(const bool bForceSendCurren
 				SessionsToDelete.Add(Session);
 			}
 
+			// NOTE: The sessions are deleted before sending (while holding the lock) to prevent another process from finding and sending the same orphan sessions twice.
+			//       This also ensures a session will never be sent twice in case a crash occurs while sending, but all sessions loaded in memory for sending will also be lost.
 			for (const FEditorAnalyticsSession& ToDelete : SessionsToDelete)
 			{
 				ToDelete.Delete();
@@ -150,8 +123,9 @@ void FEditorSessionSummarySender::SendStoredSessions(const bool bForceSendCurren
 			FEditorAnalyticsSession::Unlock();
 			bSessionsLoaded = true;
 		}
-	} while (bForceSendCurrentSession && !bSessionsLoaded); // Retry until session are loaded if the sender is forced to send the current session.
+	} while (bForceSendOwnedSession && !bSessionsLoaded); // Retry until session are loaded if the sender is forced to send the current session.
 
+	// Send the sessions (without holding the lock).
 	for (const FEditorAnalyticsSession& Session : SessionsToReport)
 	{
 		SendSessionSummaryEvent(Session);
@@ -227,14 +201,15 @@ void FEditorSessionSummarySender::SendSessionSummaryEvent(const FEditorAnalytics
 	AnalyticsAttributes.Emplace(TEXT("IsLowDriveSpace"), Session.bIsLowDriveSpace);
 	AnalyticsAttributes.Emplace(TEXT("SentFrom"), Sender);
 
-	// was this sent from some other process than itself or the out-of-process monitor for that run?
-	AnalyticsAttributes.Emplace(TEXT("DelayedSend"), Session.PlatformProcessID != CurrentSessionProcessId);
-
-	if (Session.PlatformProcessID == CurrentSessionProcessId && CurrentSessionExitCode.IsSet())
+	// Add the exit code to the report if it was set by out-of-process monitor (CrashReportClientEditor).
+	if (Session.ExitCode.IsSet())
 	{
-		AnalyticsAttributes.Emplace(TEXT("ExitCode"), CurrentSessionExitCode.GetValue());
+		AnalyticsAttributes.Emplace(TEXT("ExitCode"), Session.ExitCode.GetValue());
 	}
 
+	// Was this summary produced by another process than itself or the out-of-process monitor for that run?
+	AnalyticsAttributes.Emplace(TEXT("DelayedSend"), Session.PlatformProcessID != CurrentSessionProcessId);
+
 	// Sending the summary event of the current process analytic session?
 	if (AnalyticsProvider.GetSessionID().Contains(Session.SessionId)) // The string (GUID) returned by GetSessionID() is surrounded with braces like "{3FEA3232-...}" while Session.SessionId is not -> "3FEA3232-..."
 	{
diff --git a/Engine/Source/Developer/EditorAnalyticsSession/Public/EditorAnalyticsSession.h b/Engine/Source/Developer/EditorAnalyticsSession/Public/EditorAnalyticsSession.h
--- a/Engine/Source/Developer/EditorAnalyticsSession/Public/EditorAnalyticsSession.h
+++ b/Engine/Source/Developer/EditorAnalyticsSession/Public/EditorAnalyticsSession.h
@@ -28,7 +28,9 @@ struct EDITORANALYTICSSESSION_API FEditorAnalyticsSession
 	FString ProjectDescription;
 	FString ProjectVersion;
 	FString EngineVersion;
-	int32 PlatformProcessID;
+	uint32 PlatformProcessID;
+	uint32 MonitorProcessID; // Set to the CrashReportClientEditor PID when out-of-process reporting is used.
+	TOptional<int32> ExitCode; // Set by CrashReportClientEditor after the Editor process exit when out-of-process reporting is used and reading the exit code is supported.
 
 	FDateTime StartupTimestamp;
 	FDateTime Timestamp;
@@ -142,6 +144,20 @@ struct EDITORANALYTICSSESSION_API FEditorAnalyticsSession
 	 */
 	void LogEvent(EEventType EventTpe, const FDateTime& Timestamp);
 
+	/**
+	 * Find the sessions for which the PlatformProcessID matches the specified session process id.
+	 * @param InSessionProcessId The session with this process ID.
+	 * @param[out] The found session if any (undefined if not found)
+	 * @return True if a session for the specified process id was found, false otherwise.
+	 */
+	static bool FindSession(const uint32 InSessionProcessId, FEditorAnalyticsSession& OutSession);
+
+	/**
+	 * Persist the Editor exit code corresponding of this session in the key-store value.
+	 * @node This function should only be called by the out of process monitor when the process exit value is known.
+	 */
+	void SaveExitCode(int32 ExitCode);
+
 private:
 	static FSystemWideCriticalSection* StoredValuesLock;
 
diff --git a/Engine/Source/Developer/EditorAnalyticsSession/Public/EditorSessionSummarySender.h b/Engine/Source/Developer/EditorAnalyticsSession/Public/EditorSessionSummarySender.h
--- a/Engine/Source/Developer/EditorAnalyticsSession/Public/EditorSessionSummarySender.h
+++ b/Engine/Source/Developer/EditorAnalyticsSession/Public/EditorSessionSummarySender.h
@@ -16,15 +16,12 @@ class IAnalyticsProviderET;
 class EDITORANALYTICSSESSION_API FEditorSessionSummarySender
 {
 public:
-	FEditorSessionSummarySender(IAnalyticsProviderET& InAnalyticsProvider, const FString& InSenderName, const int32 InCurrentSessionProcessId);
+	FEditorSessionSummarySender(IAnalyticsProviderET& InAnalyticsProvider, const FString& InSenderName, const uint32 InCurrentSessionProcessId);
 	~FEditorSessionSummarySender();
 
 	void Tick(float DeltaTime);
 	void Shutdown();
 
-	void SetCurrentSessionExitCode(const int32 InCurrentSessionProcessId, const int32 InExitCode);
-	bool FindCurrentSession(FEditorAnalyticsSession& OutSession) const;
-
 private:
 	/** Send any stored Sessions. */
 	void SendStoredSessions(const bool bForceSendCurrentSession = false) const;
@@ -34,7 +31,5 @@ private:
 	float HeartbeatTimeElapsed;
 	IAnalyticsProviderET& AnalyticsProvider;
 	FString Sender;
-
-	int32 CurrentSessionProcessId;
-	TOptional<int32> CurrentSessionExitCode;
+	uint32 CurrentSessionProcessId;
 };
diff --git a/Engine/Source/Programs/CrashReportClient/Private/CrashReportClientApp.cpp b/Engine/Source/Programs/CrashReportClient/Private/CrashReportClientApp.cpp
--- a/Engine/Source/Programs/CrashReportClient/Private/CrashReportClientApp.cpp
+++ b/Engine/Source/Programs/CrashReportClient/Private/CrashReportClientApp.cpp
@@ -90,6 +90,16 @@ enum SubmitCrashReportResult {
 	SuccessContinue		// Succeeded sending report, continue running (if monitor mode).
 };
 
+/** Defines special exit codes for the Editor to use in the summary report when the real exit code could not be determined, to diagnose unexpected situations. */
+enum EditorSpecialExitCodes
+{
+	/** Editor is still running, but CRC was requested to exit. (ex. The user clicked the 'Close Without Sending' but the Editor is not exited yet) */
+	StillRunning = 777001,
+
+	/** Editor is not running anymore, but CRC could not read the Editor exit code. */
+	NotAvailable = 777002,
+};
+
 /**
  * Look for the report to upload, either in the command line or in the platform's report queue
  */
@@ -291,7 +301,7 @@ SubmitCrashReportResult RunWithUI(FPlatformErrorReport ErrorReport)
 	FCrashReportClientStyle::Initialize();
 
 	// Create the main implementation object
-	TSharedRef<FCrashReportClient> CrashReportClient = MakeShareable(new FCrashReportClient(ErrorReport));
+	TSharedRef<FCrashReportClient> CrashReportClient = MakeShared<FCrashReportClient>(ErrorReport);
 
 	// open up the app window	
 	TSharedRef<SCrashReportClient> ClientControl = SNew(SCrashReportClient, CrashReportClient);
@@ -738,23 +748,14 @@ static void HandleAbnormalShutdown(FSharedCrashContext& CrashContext, uint64 Pro
 	}
 }
 
-static bool WasAbnormalShutdown(const FEditorSessionSummarySender& SessionSummarySender)
+static bool WasAbnormalShutdown(const FEditorAnalyticsSession& AnalyticSession)
 {
-	FEditorAnalyticsSession AnalyticsSession;
-	if (SessionSummarySender.FindCurrentSession(AnalyticsSession))
-	{
-		// check if this was an abnormal shutdown (aka. none of the known shutdown types, and not debugged)
-		if (AnalyticsSession.bCrashed == false &&
-			AnalyticsSession.bGPUCrashed == false &&
-			AnalyticsSession.bWasShutdown == false &&
-			AnalyticsSession.bIsTerminating == false &&
-			AnalyticsSession.bWasEverDebugger == false)
-		{
-			return true;
-		}
-	}
-
-	return false;
+	// Check if this was an abnormal shutdown (aka. none of the known shutdown types, and not debugged)
+	return AnalyticSession.bCrashed == false &&
+		AnalyticSession.bGPUCrashed == false &&
+		AnalyticSession.bWasShutdown == false &&
+		AnalyticSession.bIsTerminating == false &&
+		AnalyticSession.bWasEverDebugger == false;
 }
 
 #endif
@@ -926,61 +927,100 @@ void RunCrashReportClient(const TCHAR* CommandLine)
 
 #if CRASH_REPORT_WITH_MTBF // Expected to be 1 when compiling CrashReportClientEditor.
 		{
+			UE_LOG(CrashReportClientLog, Log, TEXT("Starting Editor MTBF reporting"));
+
 			// The loop above can exit before the Editor (monitored process) exits (because of IsEngineExitRequested()) if the user clicks 'Close Without Sending' very quickly, but for MTBF,
 			// it is desirable to have the Editor process return code. Give some extra time to the Editor to exit. If it doesn't exit within x seconds the next CRC instance will sent the
 			// current analytic report delayed, not ideal, but supported.
-			FDateTime WaitEndTime = FDateTime::UtcNow() + FTimespan(0, 0, 120);
+			FDateTime WaitEndTime = FDateTime::UtcNow() + FTimespan::FromMinutes(3);
 			while (ProcessStatus.Get<0>() && FDateTime::UtcNow() <= WaitEndTime)
 			{
-				FPlatformProcess::Sleep(0.5f); // In seconds
+				FPlatformProcess::Sleep(0.1f); // In seconds
 				ProcessStatus = GetProcessStatus(MonitoredProcess);
 			}
 
-			// Load the temporary crash context file.
-			FSharedCrashContext TempCrashContext;
-			FMemory::Memzero(TempCrashContext);
-			if (!ProcessStatus.Get<0>() && LoadTempCrashContextFromFile(TempCrashContext, MonitorPid) && TempCrashContext.UserSettings.bSendUsageData)
+			// Check the status of the Editor process after waiting n seconds.
+			TOptional<int32> MonitoredProcessExitCode = ProcessStatus.Get<1>();
+			bool bMonitoredProcessExited = !ProcessStatus.Get<0>();
+			bool bMonitoredSessionLoaded = false;
+
+			// Try to persist an exit code in session summary (even if the Editor is still running)
+			FEditorAnalyticsSession MonitoredSession;
+			FTimespan Timeout = FTimespan::FromMinutes(2);
+			if (FEditorAnalyticsSession::Lock(Timeout))
+			{
+				if (FEditorAnalyticsSession::FindSession(MonitorPid, MonitoredSession))
+				{
+					bMonitoredSessionLoaded = true; // Were able to acquire lock and load the session.
+
+					if (MonitoredProcessExitCode.IsSet())
+					{
+						// Persist the real Editor exit code.
+						MonitoredSession.SaveExitCode(MonitoredProcessExitCode.GetValue());
+					}
+					else if (!bMonitoredProcessExited)
+					{
+						// Persist a custom exit code - Editor is still running, but CRC was requested to exit. (ex. The user clicked the 'Close Without Sending' but the Editor is not exited yet)
+						MonitoredSession.SaveExitCode(EditorSpecialExitCodes::StillRunning);
+					}
+					else
+					{
+						// Persist the custom exit code - Editor is not running anymore, but CRC could not read it.
+						MonitoredSession.SaveExitCode(EditorSpecialExitCodes::NotAvailable);
+					}
+				}
+				else
+				{
+					UE_LOG(CrashReportClientLog, Warning, TEXT("Failed to set Editor exit code. The session could not be found."));
+				}
+				FEditorAnalyticsSession::Unlock();
+			}
+			else
+			{
+				UE_LOG(CrashReportClientLog, Warning, TEXT("Failed to acquire lock to set exit code in Editor summary"));
+			}
+
+			if (bMonitoredProcessExited)
 			{
-				// Send the editor summary event(s) first, maximizing chance of sucessfull transmission (less opportunities for bugs to prevent it).
-				FCrashReportAnalytics::Initialize();
-				if (FCrashReportAnalytics::IsAvailable())
+				// Load the temporary crash context file.
+				FSharedCrashContext TempCrashContext;
+				FMemory::Memzero(TempCrashContext);
+				if (LoadTempCrashContextFromFile(TempCrashContext, MonitorPid) && TempCrashContext.UserSettings.bSendUsageData)
 				{
-					// Editor shutdown state from Editor point of view.
-					bool bAbnormalShutdownFromEditorPov = false;
+					FCrashReportAnalytics::Initialize();
+					if (FCrashReportAnalytics::IsAvailable())
 					{
+						// Send this session summary event (and the orphan ones if any).
 						FEditorSessionSummarySender EditorSessionSummarySender(FCrashReportAnalytics::GetProvider(), TEXT("CrashReportClient"), MonitorPid);
+						EditorSessionSummarySender.Shutdown();
 
-						// Query the Editor process state again, the Editor may still run.
-						ProcessStatus = GetProcessStatus(MonitoredProcess);
-						TOptional<int32> ExitCodeOpt = ProcessStatus.Get<1>();
-						if (ExitCodeOpt.IsSet()) // Exit code is known?
-						{
-							EditorSessionSummarySender.SetCurrentSessionExitCode(MonitorPid, ExitCodeOpt.GetValue()); // The Editor exit code from the OS point of view.
-						}
-						else
+						// If the Editor thinks the session ended up abnormally, generate a crash report (to get the Editor logs and figure out why this happened).
+						if (bMonitoredSessionLoaded && TempCrashContext.UserSettings.bSendUnattendedBugReports)
 						{
-							EditorSessionSummarySender.SetCurrentSessionExitCode(MonitorPid, 112233001); // Special exit code, arbitrary but easy to read in decimal, to mark the process exit code as 'unknown'.
+							// Check what the Editor knows about the exit. Was the proper handlers called and the flag(s) set in the summary event?
+							if (WasAbnormalShutdown(MonitoredSession))
+							{
+								// Send a spoofed crash report in the case that we detect an abnormal shutdown has occurred
+								HandleAbnormalShutdown(TempCrashContext, MonitorPid, MonitorWritePipe, RecoveryServicePtr);
+							}
 						}
-
-						// Check what the Editor knows about the exit. Was the proper handlers called and the flag(s) set in the summary event?
-						bAbnormalShutdownFromEditorPov = WasAbnormalShutdown(EditorSessionSummarySender);
-
-						// Send summary session event(s).
-						EditorSessionSummarySender.Shutdown();
-					}
-
-					// If the Editor thinks the session ended up abnormally, generate a crash report (to get the Editor logs and figure out why this happened).
-					if (bAbnormalShutdownFromEditorPov && TempCrashContext.UserSettings.bSendUnattendedBugReports)
-					{
-						// Send a spoofed crash report in the case that we detect an abnormal shutdown has occurred
-						HandleAbnormalShutdown(TempCrashContext, MonitorPid, MonitorWritePipe, RecoveryServicePtr);
 					}
+					FCrashReportAnalytics::Shutdown();
+				}
+				else
+				{
+					UE_LOG(CrashReportClientLog, Warning, TEXT("Failed to send Editor session summary event. Could not read the temporary file."));
 				}
-				FCrashReportAnalytics::Shutdown();
 			}
+			else
+			{
+				UE_LOG(CrashReportClientLog, Warning, TEXT("Failed to send Editor session summary event. The Editor was still running."));
+			}
+
+			UE_LOG(CrashReportClientLog, Display, TEXT("Finished Editor MTBF reporting"));
 		}
 #endif
-		// clean up the context file
+		// Clean up the context file
 		DeleteTempCrashContextFile(MonitorPid);
 
 		FPlatformProcess::CloseProc(MonitoredProcess);
diff --git a/Engine/Source/Runtime/Core/Private/GenericPlatform/GenericPlatformCrashContext.cpp b/Engine/Source/Runtime/Core/Private/GenericPlatform/GenericPlatformCrashContext.cpp
--- a/Engine/Source/Runtime/Core/Private/GenericPlatform/GenericPlatformCrashContext.cpp
+++ b/Engine/Source/Runtime/Core/Private/GenericPlatform/GenericPlatformCrashContext.cpp
@@ -69,7 +69,7 @@ const TCHAR* const FGenericCrashContext::EngineModeExDirty = TEXT("Dirty");
 const TCHAR* const FGenericCrashContext::EngineModeExVanilla = TEXT("Vanilla");
 
 bool FGenericCrashContext::bIsInitialized = false;
-bool FGenericCrashContext::bIsOutOfProcess = false;
+uint32 FGenericCrashContext::OutOfProcessCrashReporterPid = 0;
 int32 FGenericCrashContext::StaticCrashContextIndex = 0;
 
 const FGuid FGenericCrashContext::ExecutionGuid = FGuid::NewGuid();
diff --git a/Engine/Source/Runtime/Core/Private/Windows/WindowsPlatformCrashContext.cpp b/Engine/Source/Runtime/Core/Private/Windows/WindowsPlatformCrashContext.cpp
--- a/Engine/Source/Runtime/Core/Private/Windows/WindowsPlatformCrashContext.cpp
+++ b/Engine/Source/Runtime/Core/Private/Windows/WindowsPlatformCrashContext.cpp
@@ -415,7 +415,7 @@ bool CreateCrashReportClientPath(TCHAR* OutClientPath, int32 MaxLength)
 /**
  * Launches crash reporter client and creates the pipes for communication.
  */
-FProcHandle LaunchCrashReportClient(void** OutWritePipe, void** OutReadPipe)
+FProcHandle LaunchCrashReportClient(void** OutWritePipe, void** OutReadPipe, uint32* CrashReportClientProcessId)
 {
 	TCHAR CrashReporterClientPath[CR_CLIENT_MAX_PATH_LEN] = { 0 };
 	TCHAR CrashReporterClientArgs[CR_CLIENT_MAX_ARGS_LEN] = { 0 };
@@ -486,7 +486,7 @@ FProcHandle LaunchCrashReportClient(void** OutWritePipe, void** OutReadPipe)
 			CrashReporterClientPath,
 			CrashReporterClientArgs,
 			true, false, false,
-			nullptr, 0,
+			CrashReportClientProcessId, 0,
 			nullptr,
 			nullptr,
 			nullptr
@@ -968,6 +968,8 @@ private:
 	void* CrashMonitorWritePipe;
 	/** Pipe for reading from the monitor process. */
 	void* CrashMonitorReadPipe;
+	/** The crash report client process ID. */
+	uint32 CrashMonitorPid;
 	/** Memory allocated for crash context. */
 	FSharedCrashContext SharedContext;
 	
@@ -1024,6 +1026,7 @@ public:
 		, VectoredExceptionHandle(INVALID_HANDLE_VALUE)
 		, CrashMonitorWritePipe(nullptr)
 		, CrashMonitorReadPipe(nullptr)
+		, CrashMonitorPid(0)
 	{
 		// Synchronization objects
 		CrashEvent = CreateEvent(nullptr, true, 0, nullptr);
@@ -1041,7 +1044,7 @@ public:
 #if USE_CRASH_REPORTER_MONITOR
 		if (!FPlatformProperties::IsServerOnly())
 		{
-			CrashClientHandle = LaunchCrashReportClient(&CrashMonitorWritePipe, &CrashMonitorReadPipe);
+			CrashClientHandle = LaunchCrashReportClient(&CrashMonitorWritePipe, &CrashMonitorReadPipe, &CrashMonitorPid);
 			FMemory::Memzero(SharedContext);
 		}
 #endif
@@ -1053,7 +1056,10 @@ public:
 			SetThreadPriority(Thread, THREAD_PRIORITY_BELOW_NORMAL);
 		}
 
-		FGenericCrashContext::SetIsOutOfProcessCrashReporter(CrashClientHandle.IsValid());
+		if (CrashClientHandle.IsValid())
+		{
+			FGenericCrashContext::SetOutOfProcessCrashReporterPid(CrashMonitorPid);
+		}
 	}
 
 	FORCENOINLINE ~FCrashReportingThread()
diff --git a/Engine/Source/Runtime/Core/Public/GenericPlatform/GenericPlatformCrashContext.h b/Engine/Source/Runtime/Core/Public/GenericPlatform/GenericPlatformCrashContext.h
--- a/Engine/Source/Runtime/Core/Public/GenericPlatform/GenericPlatformCrashContext.h
+++ b/Engine/Source/Runtime/Core/Public/GenericPlatform/GenericPlatformCrashContext.h
@@ -275,13 +275,21 @@ public:
 	 */
 	static bool IsOutOfProcessCrashReporter()
 	{
-		return bIsOutOfProcess;
+		return OutOfProcessCrashReporterPid != 0;
+	}
+
+	/**
+	 * @return a non-zero value if crash reporter process is used to monitor the session or zero for in-process reporting.
+	 */
+	static uint32 GetOutOfProcessCrashReporterProcessId()
+	{
+		return OutOfProcessCrashReporterPid;
 	}
 
 	/** Set whether or not the out-of-process crash reporter is running. */
-	static void SetIsOutOfProcessCrashReporter(bool bInValue)
+	static void SetOutOfProcessCrashReporterPid(uint32 ProcessId)
 	{
-		bIsOutOfProcess = bInValue;
+		OutOfProcessCrashReporterPid = ProcessId;
 	}
 
 	/** Default constructor. Optionally pass a process handle if building a crash context for a process other then current. */
@@ -483,8 +491,8 @@ private:
 	/**	Whether the Initialize() has been called */
 	static bool bIsInitialized;
 
-	/** Whether or not crash reporting is being handled out-of-process. */
-	static bool bIsOutOfProcess;
+	/** The ID of the external process reporting crashes if the platform supports it and was configured to use it, zero otherwise (0 is a reserved system process ID, invalid for the out of process reporter). */
+	static uint32 OutOfProcessCrashReporterPid;
 
 	/**	Static counter records how many crash contexts have been constructed */
 	static int32 StaticCrashContextIndex;
diff --git a/Engine/Source/Runtime/Engine/Private/Analytics/EditorSessionSummaryWriter.cpp b/Engine/Source/Runtime/Engine/Private/Analytics/EditorSessionSummaryWriter.cpp
--- a/Engine/Source/Runtime/Engine/Private/Analytics/EditorSessionSummaryWriter.cpp
+++ b/Engine/Source/Runtime/Engine/Private/Analytics/EditorSessionSummaryWriter.cpp
@@ -37,9 +37,10 @@ namespace EditorSessionWriterDefs
 	static const float HeartbeatPeriodSeconds = 60;
 }
 
-FEditorSessionSummaryWriter::FEditorSessionSummaryWriter()
+FEditorSessionSummaryWriter::FEditorSessionSummaryWriter(uint32 InOutOfProcessMonitorProcessId)
 	: HeartbeatTimeElapsed(0.0f)
 	, bShutdown(false)
+	, OutOfProcessMonitorProcessId(InOutOfProcessMonitorProcessId)
 {
 }
 
@@ -60,7 +61,7 @@ void FEditorSessionSummaryWriter::Initialize()
 	if (FEditorAnalyticsSession::TryLock()) // System wide lock to write the session file/registry. Don't block if already taken, delay initialisation to the next Tick().
 	{
 		// Create a session Session for this session
-		CurrentSession = CreateCurrentSession();
+		CurrentSession = CreateCurrentSession(OutOfProcessMonitorProcessId);
 		CurrentSession->Save();
 
 		UE_LOG(LogEditorSessionSummary, Log, TEXT("EditorSessionSummaryWriter initialized"));
@@ -206,7 +207,7 @@ void FEditorSessionSummaryWriter::Shutdown()
 	bShutdown = true;
 }
 
-TUniquePtr<FEditorAnalyticsSession> FEditorSessionSummaryWriter::CreateCurrentSession()
+TUniquePtr<FEditorAnalyticsSession> FEditorSessionSummaryWriter::CreateCurrentSession(uint32 OutOfProcessMonitorProcessId)
 {
 	check(FEngineAnalytics::IsAvailable()); // The function assumes the caller checked it before calling.
 
@@ -232,6 +233,7 @@ TUniquePtr<FEditorAnalyticsSession> FEditorSessionSummaryWriter::CreateCurrentSe
 	Session->UserId = AnalyticProvider.GetUserID();
 
 	Session->PlatformProcessID = FPlatformProcess::GetCurrentProcessId();
+	Session->MonitorProcessID = OutOfProcessMonitorProcessId;
 	Session->ProjectName = ProjectSettings.ProjectName;
 	Session->ProjectID = ProjectSettings.ProjectID.ToString(EGuidFormats::DigitsWithHyphens);
 	Session->ProjectDescription = ProjectSettings.Description;
diff --git a/Engine/Source/Runtime/Engine/Private/Analytics/EditorSessionSummaryWriter.h b/Engine/Source/Runtime/Engine/Private/Analytics/EditorSessionSummaryWriter.h
--- a/Engine/Source/Runtime/Engine/Private/Analytics/EditorSessionSummaryWriter.h
+++ b/Engine/Source/Runtime/Engine/Private/Analytics/EditorSessionSummaryWriter.h
@@ -13,7 +13,7 @@ struct FEditorAnalyticsSession;
 class FEditorSessionSummaryWriter
 {
 public:
-	FEditorSessionSummaryWriter();
+	FEditorSessionSummaryWriter(uint32 OutOfProcessMonitorProcessId);
 	~FEditorSessionSummaryWriter();
 
 	void Initialize();
@@ -27,7 +27,7 @@ private:
 	void OnUserActivity(const FUserActivity& UserActivity);
 	void OnVanillaStateChanged(bool bIsVanilla);
 
-	static TUniquePtr<FEditorAnalyticsSession> CreateCurrentSession();
+	static TUniquePtr<FEditorAnalyticsSession> CreateCurrentSession(uint32 OutOfProcessMonitorProcessId);
 	static FString GetUserActivityString();
 	void UpdateTimestamps();
 	void UpdateIdleTimes();
@@ -39,6 +39,7 @@ private:
 	FCriticalSection SaveSessionLock;
 	float HeartbeatTimeElapsed;
 	bool bShutdown;
+	const uint32 OutOfProcessMonitorProcessId; // Non-zero if out-of process monitoring is set. To ensure one CrashReportClient(CRC) doesn't report the session of another CRC instance (race condition).
 };
 
 #endif
diff --git a/Engine/Source/Runtime/Engine/Private/EngineAnalytics.cpp b/Engine/Source/Runtime/Engine/Private/EngineAnalytics.cpp
--- a/Engine/Source/Runtime/Engine/Private/EngineAnalytics.cpp
+++ b/Engine/Source/Runtime/Engine/Private/EngineAnalytics.cpp
@@ -291,14 +291,14 @@ void FEngineAnalytics::Initialize()
 		// Create the session manager singleton
 		if (!SessionManager.IsValid())
 		{
-			SessionManager = MakeShareable(new FEngineSessionManager(EEngineSessionManagerMode::Editor));
+			SessionManager = MakeShared<FEngineSessionManager>(EEngineSessionManagerMode::Editor);
 			SessionManager->Initialize();
 		}
 
 #if WITH_EDITOR
 		if (!SessionSummaryWriter.IsValid())
 		{
-			SessionSummaryWriter = MakeShareable(new FEditorSessionSummaryWriter());
+			SessionSummaryWriter = MakeShared<FEditorSessionSummaryWriter>(FGenericCrashContext::GetOutOfProcessCrashReporterProcessId());
 			SessionSummaryWriter->Initialize();
 		}
 
