Allow CRC analytics settings to be compiled into the executable, to remove NotForLicensees folder within engine code.

diff --git a/Engine/Build/InstalledEngineBuild.xml b/Engine/Build/InstalledEngineBuild.xml
--- a/Engine/Build/InstalledEngineBuild.xml
+++ b/Engine/Build/InstalledEngineBuild.xml
@@ -124,7 +124,7 @@
 
 	<Option Name="CrashReporterAPIURL" DefaultValue="" Description="The URL to use to talk to the CrashReporterClient API."/>
 	<Option Name="CrashReporterAPIKey" DefaultValue="" Description="The API key to use to talk to the CrashReporterClient API."/>
-	<Property Name="CrashReporterCompileArgs" Value="-CompilerArguments=&quot;/DCRC_TELEMETRY_URL=$(CrashReporterAPIURL) /DCRC_TELEMETRY_KEY=$(CrashReporterAPIKey)&quot;" If="'$(CrashReporterAPIURL)' != '' And '$(CrashReporterAPIKey)' != ''" />
+	<Property Name="CrashReporterCompileArgs" Value="-CompilerArguments=&quot;/DCRC_TELEMETRY_URL=$(CrashReporterAPIURL) /DCRC_TELEMETRY_KEY_DEV=$(CrashReporterAPIKey) /DCRC_TELEMETRY_KEY_RELEASE=$(CrashReporterAPIKey)&quot;" If="'$(CrashReporterAPIURL)' != '' And '$(CrashReporterAPIKey)' != ''" />
 
 	<!-- Base editor and tools for compile on Windows -->
 	<Agent Name="Editor Win64" Type="Win64_Licensee">
diff --git a/Engine/Source/Programs/CrashReportClient/CrashReportClient.Target.cs b/Engine/Source/Programs/CrashReportClient/CrashReportClient.Target.cs
--- a/Engine/Source/Programs/CrashReportClient/CrashReportClient.Target.cs
+++ b/Engine/Source/Programs/CrashReportClient/CrashReportClient.Target.cs
@@ -1,12 +1,23 @@
 // Copyright Epic Games, Inc. All Rights Reserved.
 
-using UnrealBuildTool;
+using System;
 using System.Collections.Generic;
+using System.Linq;
+using UnrealBuildTool;
 
 [SupportedPlatforms("Win32", "Win64", "Mac", "Linux")]
 [SupportedConfigurations(UnrealTargetConfiguration.Debug, UnrealTargetConfiguration.Development, UnrealTargetConfiguration.Shipping)]
 public class CrashReportClientTarget : TargetRules
 {
+	[ConfigFile(ConfigHierarchyType.Engine, "CrashReportClientBuildSettings", "TelemetryUrl")]
+	public string TelemetryUrl;
+
+	[ConfigFile(ConfigHierarchyType.Engine, "CrashReportClientBuildSettings", "TelemetryKey_Dev")]
+	public string TelemetryKey_Dev;
+
+	[ConfigFile(ConfigHierarchyType.Engine, "CrashReportClientBuildSettings", "TelemetryKey_Release")]
+	public string TelemetryKey_Release;
+
 	public CrashReportClientTarget(TargetInfo Target) : base(Target)
 	{
 		Type = TargetType.Program;
@@ -41,6 +52,19 @@ public class CrashReportClientTarget : TargetRules
 		// Need to disable the bundled version of dbghelp so that CrashDebugHelper can load dbgeng.dll.
 		WindowsPlatform.bUseBundledDbgHelp = false;
 
+		// Add the definitions from config files
+		AddConfigMacro("CRC_TELEMETRY_URL=", TelemetryUrl);
+		AddConfigMacro("CRC_TELEMETRY_KEY_DEV=", TelemetryKey_Dev);
+		AddConfigMacro("CRC_TELEMETRY_KEY_RELEASE=", TelemetryKey_Release);
+
 		GlobalDefinitions.Add("NOINITCRASHREPORTER=1");
 	}
+
+	void AddConfigMacro(string Prefix, string Value)
+	{
+		if (!string.IsNullOrEmpty(Value) && !GlobalDefinitions.Any(x => x.StartsWith(Prefix, StringComparison.Ordinal)))
+		{
+			GlobalDefinitions.Add(Prefix + Value);
+		}
+	}
 }
diff --git a/Engine/Source/Runtime/CrashReportCore/Private/CrashReportAnalytics.cpp b/Engine/Source/Runtime/CrashReportCore/Private/CrashReportAnalytics.cpp
--- a/Engine/Source/Runtime/CrashReportCore/Private/CrashReportAnalytics.cpp
+++ b/Engine/Source/Runtime/CrashReportCore/Private/CrashReportAnalytics.cpp
@@ -2,9 +2,11 @@
 
 #include "CrashReportAnalytics.h"
 #include "Misc/Guid.h"
+#include "Misc/EngineBuildSettings.h"
 
 #include "AnalyticsET.h"
 #include "IAnalyticsProviderET.h"
+#include "AnalyticsBuildType.h"
 
 bool FCrashReportAnalytics::bIsInitialized;
 TSharedPtr<IAnalyticsProviderET> FCrashReportAnalytics::Analytics;
@@ -42,11 +44,17 @@ void FCrashReportAnalytics::Initialize()
 	checkf(!bIsInitialized, TEXT("FCrashReportAnalytics::Initialize called more than once."));
 
 	// Allow build machines to force CRC to enable internal telemetry.
-	#if defined(CRC_TELEMETRY_URL) && defined(CRC_TELEMETRY_KEY)
+	#if defined(CRC_TELEMETRY_URL) && defined(CRC_TELEMETRY_KEY_DEV) && defined(CRC_TELEMETRY_KEY_RELEASE)
+
+		// We always use the "Release" analytics account unless we're running in analytics test mode (usually with
+		// a command-line parameter), or we're an internal Epic build
+		bool bUseReleaseAccount =
+			(GetAnalyticsBuildType() == EAnalyticsBuildType::Development || GetAnalyticsBuildType() == EAnalyticsBuildType::Release) &&
+			!FEngineBuildSettings::IsInternalBuild();	// Internal Epic build
 
 		FAnalyticsET::Config Config;
 		Config.APIServerET = TEXT(PREPROCESSOR_TO_STRING(CRC_TELEMETRY_URL));
-		Config.APIKeyET = TEXT(PREPROCESSOR_TO_STRING(CRC_TELEMETRY_KEY));
+		Config.APIKeyET = bUseReleaseAccount ? TEXT(PREPROCESSOR_TO_STRING(CRC_TELEMETRY_KEY_RELEASE)) : TEXT(PREPROCESSOR_TO_STRING(CRC_TELEMETRY_KEY_DEV));
 
 	#else
 
