Remove unnecessary whitelist from CrashReportClient module.

diff --git a/Engine/Source/Programs/CrashReportClient/CrashReportClient.Build.cs b/Engine/Source/Programs/CrashReportClient/CrashReportClient.Build.cs
--- a/Engine/Source/Programs/CrashReportClient/CrashReportClient.Build.cs
+++ b/Engine/Source/Programs/CrashReportClient/CrashReportClient.Build.cs
@@ -55,7 +55,5 @@ public class CrashReportClient : ModuleRules
 		}
 
 		PrivateIncludePaths.Add("Runtime/Launch/Private");		// For LaunchEngineLoop.cpp include
-
-		WhitelistRestrictedFolders.Add("Private/NotForLicensees");
 	}
 }
