Fixed CIS error C4459: declaration of 'MonitorPid' hides global declaration

diff --git a/Engine/Source/Programs/CrashReportClient/Private/Windows/CrashReportClientMainWindows.cpp b/Engine/Source/Programs/CrashReportClient/Private/Windows/CrashReportClientMainWindows.cpp
--- a/Engine/Source/Programs/CrashReportClient/Private/Windows/CrashReportClientMainWindows.cpp
+++ b/Engine/Source/Programs/CrashReportClient/Private/Windows/CrashReportClientMainWindows.cpp
@@ -24,14 +24,14 @@ int WINAPI WinMain(_In_ HINSTANCE hInInstance, _In_opt_ HINSTANCE hPrevInstance,
 	}
 	__except(EXCEPTION_EXECUTE_HANDLER)
 	{
-		uint64 MonitorPid;
-		if (FParse::Value(GetCommandLineW(), TEXT("-MONITOR="), MonitorPid))
+		uint64 MonitoredEditorPid;
+		if (FParse::Value(GetCommandLineW(), TEXT("-MONITOR="), MonitoredEditorPid))
 		{
 			FTimespan Timeout = FTimespan::FromSeconds(5);
 			if (FEditorAnalyticsSession::Lock(Timeout))
 			{
 				FEditorAnalyticsSession MonitoredSession;
-				if (FEditorAnalyticsSession::FindSession(MonitorPid, MonitoredSession))
+				if (FEditorAnalyticsSession::FindSession(MonitoredEditorPid, MonitoredSession))
 				{
 					MonitoredSession.SaveMonitorExceptCode(GetExceptionCode());
 				}
