Fix CIS reporting 'CRASH_REPORT_WITH_MTBF' is not defined as a preprocessor macro, replacing with '0' for '#if/#elif'

#rnx

diff --git a/Engine/Source/Programs/CrashReportClient/Private/Windows/CrashReportClientMainWindows.cpp b/Engine/Source/Programs/CrashReportClient/Private/Windows/CrashReportClientMainWindows.cpp
--- a/Engine/Source/Programs/CrashReportClient/Private/Windows/CrashReportClientMainWindows.cpp
+++ b/Engine/Source/Programs/CrashReportClient/Private/Windows/CrashReportClientMainWindows.cpp
@@ -3,7 +3,7 @@
 #include "CrashReportClientApp.h"
 #include "Windows/WindowsHWrapper.h"
 
-#if CRASH_REPORT_WITH_MTBF
+#if defined(CRASH_REPORT_WITH_MTBF)
 #include "EditorAnalyticsSession.h"
 #endif
 
@@ -14,7 +14,7 @@ int WINAPI WinMain(_In_ HINSTANCE hInInstance, _In_opt_ HINSTANCE hPrevInstance,
 {
 	hInstance = hInInstance;
 
-#if CRASH_REPORT_WITH_MTBF && !PLATFORM_SEH_EXCEPTIONS_DISABLED
+#if defined(CRASH_REPORT_WITH_MTBF) && !PLATFORM_SEH_EXCEPTIONS_DISABLED
 	// Try to record if CrashReportClientEditor is crashing. Analytics shows that good number of Editor exit code are reported delayed, hinting
 	// that CRCEditor was not running anymore. Try figuring out if it crashed. Suspecting that the Editor crash reporter/handler code is crashing could also
 	// inadvertedly cause a crash in CRCEditor.
