- Added code to the Editor to detect and report when CrashReportClientEditor exited unexpectedly. (MonitorExceptCode 777005 is set in the Editor session summary event)
  - Added a retrial loop to CrashReportClientApp to retry opening the the handle on the Editor process if the first time fails.

diff --git a/Engine/Source/Programs/CrashReportClient/Private/CrashReportClientApp.cpp b/Engine/Source/Programs/CrashReportClient/Private/CrashReportClientApp.cpp
--- a/Engine/Source/Programs/CrashReportClient/Private/CrashReportClientApp.cpp
+++ b/Engine/Source/Programs/CrashReportClient/Private/CrashReportClientApp.cpp
@@ -814,12 +814,26 @@ void RunCrashReportClient(const TCHAR* CommandLine)
 		// Starts the disaster recovery service. This records transactions and allows users to recover from previous crashes.
 		RecoveryServicePtr = MakeShared<FRecoveryService>(MonitorPid);
 #endif
+
+		// Try to open the process. The analytics shows that CRC sometime exits before the Editor. This can happen if OpenProcess() fails, maybe because
+		// this process was spawned in the Editor Pre-init phase and the OS did not make it available yet. In case of failure, give it a few trials.
+		FProcHandle MonitoredProcess;
+		for (int Trial = 0; Trial < 5; ++Trial)
+		{
 #if PLATFORM_WINDOWS
-		// We do not need to open a full access process handle when monitoring process health.
-		FProcHandle MonitoredProcess = FProcHandle(::OpenProcess(PROCESS_DUP_HANDLE | PROCESS_QUERY_LIMITED_INFORMATION | PROCESS_TERMINATE | SYNCHRONIZE, 0, MonitorPid));
+			// We do not need to open a full access 
+			MonitoredProcess = FProcHandle(::OpenProcess(PROCESS_DUP_HANDLE | PROCESS_QUERY_LIMITED_INFORMATION | PROCESS_TERMINATE | SYNCHRONIZE, 0, MonitorPid));
 #else
-		FProcHandle MonitoredProcess = FPlatformProcess::OpenProcess(MonitorPid);
+			MonitoredProcess = FPlatformProcess::OpenProcess(MonitorPid);
 #endif
+			if (MonitoredProcess.IsValid())
+			{
+				break;
+			}
+
+			FPlatformProcess::Sleep(1); // Wait a seconds before retrying. Give time to the Editor to initialize.
+		}
+
 		if (!MonitoredProcess.IsValid())
 		{
 			UE_LOG(CrashReportClientLog, Error, TEXT("Failed to open monitor process handle!"));
diff --git a/Engine/Source/Runtime/Core/Public/GenericPlatform/GenericPlatformCrashContext.h b/Engine/Source/Runtime/Core/Public/GenericPlatform/GenericPlatformCrashContext.h
--- a/Engine/Source/Runtime/Core/Public/GenericPlatform/GenericPlatformCrashContext.h
+++ b/Engine/Source/Runtime/Core/Public/GenericPlatform/GenericPlatformCrashContext.h
@@ -24,6 +24,9 @@ enum ECrashExitCodes : int32
 
 	/** Used by the application when the crash handler crashed itself (crash in the __except() clause for example).*/
 	CrashHandlerCrashed = 777004,
+
+	/** Used by the application to flag when it detects that its out-of-process applicate supposed to report the bugs died (ex if CrashReportClientEditor dies before the Editor).*/
+	OutOfProcessReporterExitedUnexpectedly = 777005,
 };
 
 /** 
diff --git a/Engine/Source/Runtime/Engine/Private/Analytics/EditorSessionSummaryWriter.cpp b/Engine/Source/Runtime/Engine/Private/Analytics/EditorSessionSummaryWriter.cpp
--- a/Engine/Source/Runtime/Engine/Private/Analytics/EditorSessionSummaryWriter.cpp
+++ b/Engine/Source/Runtime/Engine/Private/Analytics/EditorSessionSummaryWriter.cpp
@@ -23,6 +23,7 @@
 #include "Misc/EngineBuildSettings.h"
 #include "Misc/EngineVersion.h"
 #include "Misc/Guid.h"
+#include "GenericPlatform/GenericPlatformCrashContext.h"
 
 #include "EditorAnalyticsSession.h"
 #include "IVREditorModule.h"
@@ -253,6 +254,16 @@ void FEditorSessionSummaryWriter::Tick(float DeltaTime)
 		{
 			HeartbeatTimeElapsed = 0.0f;
 
+			// Check if the out of process monitor is running.
+			if (CurrentSession->MonitorProcessID != 0 && !CurrentSession->MonitorExceptCode.IsSet())
+			{
+				// The out-of-process application reporting our crash shouldn't die before this process.
+				if (!FPlatformProcess::IsApplicationRunning(CurrentSession->MonitorProcessID))
+				{
+					CurrentSession->MonitorExceptCode.Emplace(ECrashExitCodes::OutOfProcessReporterExitedUnexpectedly);
+				}
+			}
+
 			// check if the debugger is present
 			bool bIsDebuggerPresent = FPlatformMisc::IsDebuggerPresent();
 			if (CurrentSession->bIsDebugger != bIsDebuggerPresent)
@@ -391,6 +402,12 @@ TUniquePtr<FEditorAnalyticsSession> FEditorSessionSummaryWriter::CreateCurrentSe
 
 	Session->Plugins.Sort();
 
+	// The out-of-process application reporting our crash shouldn't die before this process.
+	if (Session->MonitorProcessID != 0 && !FPlatformProcess::IsApplicationRunning(Session->MonitorProcessID))
+	{
+		Session->MonitorExceptCode.Emplace(ECrashExitCodes::OutOfProcessReporterExitedUnexpectedly);
+	}
+
 	return Session;
 }
 
