- Implemented a special logger inside CrashReportClientEditor to capture and save important events such as crash reporting (along with the CrashGUID)
  - When CrashReportClientEditor sends all the Editor summary events, if an error was detected in the session being sent, the mini-log for that session is attached to the analytic event.

diff --git a/Engine/Source/Developer/EditorAnalyticsSession/Private/EditorSessionSummarySender.cpp b/Engine/Source/Developer/EditorAnalyticsSession/Private/EditorSessionSummarySender.cpp
--- a/Engine/Source/Developer/EditorAnalyticsSession/Private/EditorSessionSummarySender.cpp
+++ b/Engine/Source/Developer/EditorAnalyticsSession/Private/EditorSessionSummarySender.cpp
@@ -56,6 +56,11 @@ void FEditorSessionSummarySender::Shutdown()
 	SendStoredSessions(/*bForceSendCurrentSession*/true);
 }
 
+void FEditorSessionSummarySender::SetMonitorDiagnosticLogs(TMap<uint32, FString>&& Logs)
+{
+	MonitorMiniLogs = Logs;
+}
+
 void FEditorSessionSummarySender::SendStoredSessions(const bool bForceSendOwnedSession) const
 {
 	// Load the list of sessions to process. Expect contention on the analytic session lock between the Editor and CrashReportClientEditor (on windows) or between Editor instances (on mac/linux)
@@ -203,16 +208,30 @@ void FEditorSessionSummarySender::SendSessionSummaryEvent(const FEditorAnalytics
 	AnalyticsAttributes.Emplace(TEXT("IsLowDriveSpace"), Session.bIsLowDriveSpace);
 	AnalyticsAttributes.Emplace(TEXT("SentFrom"), Sender);
 
+	bool bErrorDetected = (ShutdownTypeString != EditorSessionSenderDefs::ShutdownSessionToken && ShutdownTypeString != EditorSessionSenderDefs::TerminatedSessionToken);
+
 	// Add the exit code to the report if it was set by out-of-process monitor (CrashReportClientEditor).
 	if (Session.ExitCode.IsSet())
 	{
 		AnalyticsAttributes.Emplace(TEXT("ExitCode"), Session.ExitCode.GetValue());
+		bErrorDetected |= Session.ExitCode.GetValue() != 0;
 	}
 
 	// Add the monitor exception code in case the out-of-process monitor (CrashReportClientEditor) crashed itself, caught the exception and was able to store it in the session before dying.
 	if (Session.MonitorExceptCode.IsSet())
 	{
 		AnalyticsAttributes.Emplace(TEXT("MonitorExceptCode"), Session.MonitorExceptCode.GetValue());
+		bErrorDetected = true;
+	}
+
+	// If the session did not end normally, try to attache the mini-log created for that session to help diagnose abnormal terminations.
+	if (bErrorDetected && !Session.bWasEverDebugger)
+	{
+		// Check if a monitor log is available. (Set by CrashReportClientEditor before sending the summary events).
+		if (const FString* MonitorLog = MonitorMiniLogs.Find(Session.MonitorProcessID))
+		{
+			AnalyticsAttributes.Emplace(TEXT("MonitorLog"), *MonitorLog);
+		}
 	}
 
 	// Was this summary produced by another process than itself or the out-of-process monitor for that run?
diff --git a/Engine/Source/Developer/EditorAnalyticsSession/Public/EditorSessionSummarySender.h b/Engine/Source/Developer/EditorAnalyticsSession/Public/EditorSessionSummarySender.h
--- a/Engine/Source/Developer/EditorAnalyticsSession/Public/EditorSessionSummarySender.h
+++ b/Engine/Source/Developer/EditorAnalyticsSession/Public/EditorSessionSummarySender.h
@@ -5,6 +5,7 @@
 #include "Containers/UnrealString.h"
 #include "HAL/Platform.h"
 #include "Misc/Optional.h"
+#include "Containers/Map.h"
 
 struct FEditorAnalyticsSession;
 class IAnalyticsProviderET;
@@ -22,6 +23,8 @@ public:
 	void Tick(float DeltaTime);
 	void Shutdown();
 
+	void SetMonitorDiagnosticLogs(TMap<uint32, FString>&& Logs);
+
 private:
 	/** Send any stored Sessions. */
 	void SendStoredSessions(const bool bForceSendCurrentSession = false) const;
@@ -32,4 +35,5 @@ private:
 	IAnalyticsProviderET& AnalyticsProvider;
 	FString Sender;
 	uint32 CurrentSessionProcessId;
+	TMap<uint32, FString> MonitorMiniLogs; // Maps Monitor Process ID/Monitor Log
 };
diff --git a/Engine/Source/Programs/CrashReportClient/Private/CrashReportClientApp.cpp b/Engine/Source/Programs/CrashReportClient/Private/CrashReportClientApp.cpp
--- a/Engine/Source/Programs/CrashReportClient/Private/CrashReportClientApp.cpp
+++ b/Engine/Source/Programs/CrashReportClient/Private/CrashReportClientApp.cpp
@@ -10,6 +10,9 @@
 #include "Misc/ConfigCacheIni.h"
 #include "GenericPlatform/GenericApplication.h"
 #include "Misc/App.h"
+#include "Misc/CString.h"
+#include "Misc/Paths.h"
+#include "Misc/FileHelper.h"
 #include "CrashReportCoreConfig.h"
 #include "GenericPlatform/GenericPlatformCrashContext.h"
 #include "CrashDescription.h"
@@ -17,8 +20,11 @@
 #include "Modules/ModuleManager.h"
 #include "HAL/PlatformApplicationMisc.h"
 #include "HAL/PlatformCrashContext.h"
+#include "HAL/PlatformProcess.h"
+#include "HAL/FileManager.h"
 #include "IAnalyticsProviderET.h"
 #include "XmlParser.h"
+#include "Containers/Map.h"
 
 #if !CRASH_REPORT_UNATTENDED_ONLY
 	#include "SCrashReportClient.h"
@@ -76,6 +82,154 @@ static bool AnalyticsEnabledFromCmd = true;
 /** If in monitor mode, watch this pid. */
 static uint64 MonitorPid = 0;
 
+/**
+ * Write a mini log of important events such as the crash GUID(s) to attach to the editor summary analytic event. This is to diagnose abnormal terminations
+ * of the Editor or crash in CRC itself. Each log entry is expected to be small and concise. Each instance creates a single log file identified by the CRC
+ * process ID. When CRC (compiled with MTBF support) is about to exit, it reloads the existing log files and pass them to EditorSessionSummarySender. When
+ * the summary sender is about to send a session, it check the session status, if an error occurred, it tries to match a mini-log and if the corresponding
+ * log is found, it is attached to the session summary.
+ */
+class FDiagnosticLogger
+{
+public:
+	static constexpr int32 MaxLogLen = 8 * 1024;
+
+	static FDiagnosticLogger& Get()
+	{
+		static FDiagnosticLogger Instance;
+		return Instance;
+	}
+
+	/** Log a small events to help diagnose abnormal shutdown or bugs in CRC itself. The event text is expected to be short and concise. */
+	void LogEvent(const TCHAR* Event)
+	{
+		if (IsEnabled())
+		{
+			static const TCHAR* Separator = TEXT("|");
+			static const int32 SeparatorLen = FCString::Strlen(Separator);
+
+			// Rotate the log if it gets too long.
+			int32 EventLen =  SeparatorLen + FCString::Strlen(Event); // Always account for the separator. If the diagnostic string is empty, it doesn't make any difference.
+			if (DiagnosticInfo.Len() + EventLen >= MaxLogLen)
+			{
+				DiagnosticInfo.RemoveAt(0, EventLen, /*bAllowShrinking*/false); // Remove oldest event (in front).
+			}
+
+			// Log the last event.
+			if (DiagnosticInfo.Len())
+			{
+				DiagnosticInfo.Append(Separator);
+			}
+			DiagnosticInfo.Append(Event);
+			FFileHelper::SaveStringToFile(DiagnosticInfo, *GetLogPathname());
+		}
+	}
+
+	static TMap<uint32, FString> LoadAllLogs()
+	{
+		TMap<uint32, FString> Logs;
+
+		IFileManager::Get().IterateDirectory(*GetLogDir(), [&Logs](const TCHAR* Pathname, bool bIsDir)
+		{
+			if (!bIsDir)
+			{
+				FString Filename = FPaths::GetCleanFilename(Pathname);
+				if (Filename.StartsWith(GetBaseFilename()) && Filename.EndsWith(TEXT(".log")))
+				{
+					uint32 ProcessID = GetLogProcessId(Filename);
+					if (ProcessID == FPlatformProcess::GetCurrentProcessId() || !FPlatformProcess::IsApplicationRunning(ProcessID)) // Don't load the log of another running CrashReportClient.
+					{
+						FString MonitorLog;
+						if (FFileHelper::LoadFileToString(MonitorLog, Pathname))
+						{
+							Logs.Emplace(ProcessID, MoveTemp(MonitorLog));
+						}
+					}
+				}
+			}
+			return true; // Continue iterating the directory.
+		});
+
+		return Logs;
+	}
+
+	static void ClearAllLogs()
+	{
+		IFileManager::Get().IterateDirectory(*GetLogDir(), [](const TCHAR* Pathname, bool bIsDir)
+		{
+			if (!bIsDir)
+			{
+				FString Filename = FPaths::GetCleanFilename(Pathname);
+				if (Filename.StartsWith(GetBaseFilename()) && Filename.EndsWith(TEXT(".log")))
+				{
+					uint32 ProcessID = GetLogProcessId(Filename);
+					if (ProcessID == FPlatformProcess::GetCurrentProcessId() || !FPlatformProcess::IsApplicationRunning(ProcessID)) // Don't delete the log of another running CrashReportClient.
+					{
+						IFileManager::Get().Delete(Pathname);
+					}
+				}
+			}
+			return true; // Continue iterating the directory.
+		});
+	}
+
+private:
+	FDiagnosticLogger()
+	{
+		if (IsEnabled())
+		{
+			// Ensure the Log directory exists.
+			IFileManager::Get().MakeDirectory(*GetLogDir(), /*Tree*/true);
+			DiagnosticInfo.Reset(MaxLogLen);
+		}
+	}
+
+	static bool IsEnabled()
+	{
+		// Only log if MTBF is enabled. In this mode, the mini-log created is attached to the Editor session summary to diagnose problems in CRC or help
+		// figure out Editor abnormal terminations.
+		return CRASH_REPORT_WITH_MTBF != 0;
+	}
+
+	static const FString& GetLogDir()
+	{
+		static FString LogDir(FPaths::ProjectSavedDir() / TEXT("Logs"));
+		return LogDir;
+	}
+
+	static const TCHAR* GetBaseFilename()
+	{
+		return TEXT("CrcDiagLog");
+	}
+
+	static const FString& GetLogPathname()
+	{
+		static FString LogPathname(GetLogDir() / FString::Printf(TEXT("%s_%s.log"), GetBaseFilename(), *LexToString(FPlatformProcess::GetCurrentProcessId())));
+		return LogPathname;
+	}
+
+	static uint32 GetLogProcessId(const FString& Filename)
+	{
+		// Parse the PID from a filename like: CrcDiagLog_939399.log
+		int Start;
+		int End;
+		if (!Filename.FindChar(TEXT('_'), Start))
+		{
+			return 0;
+		}
+		else if (!Filename.FindChar(TEXT('.'), End))
+		{
+			return 0;
+		}
+
+		FString ProcessIdStr = Filename.Mid(Start + 1, End - Start);
+		return FCString::Atoi(*ProcessIdStr);
+	}
+
+private:
+	FString DiagnosticInfo;
+};
+
 /** If in monitor mode, pipe to read data from game. */
 static void* MonitorReadPipe = nullptr;
 
@@ -493,6 +647,9 @@ FPlatformErrorReport CollectErrorReport(FRecoveryService* RecoveryService, uint3
 
 	FPlatformErrorReport ErrorReport(ReportDirectoryAbsolutePath);
 
+	// Link the crash to the Editor summary event to help diagnose the abnormal termination quickly.
+	FDiagnosticLogger::Get().LogEvent(*FPrimaryCrashProperties::Get()->CrashGUID);
+
 #if CRASH_REPORT_UNATTENDED_ONLY
 	return ErrorReport;
 #else
@@ -594,10 +751,27 @@ bool IsCrashReportAvailable(uint32 WatchedProcess, FSharedCrashContext& CrashCon
 		if (TotalRead < sizeof(FSharedCrashContext))
 		{
 			UE_LOG(CrashReportClientLog, Error, TEXT("The shared crash context emitted by the monitored process could not be fully read (or was smaller than expected by crash reporter)."));
+			FDiagnosticLogger::Get().LogEvent(TEXT("Pipe/NotEnoughData"));
 		}
 		else if (TotalRead > sizeof(FSharedCrashContext))
 		{
 			UE_LOG(CrashReportClientLog, Error, TEXT("The shared crash context emitted by the monitored process is larger than expected by crash reporter."));
+			FDiagnosticLogger::Get().LogEvent(TEXT("Pipe/TooMuchData"));
+		}
+		else
+		{
+			// Record the history of events sent by the Editor to help diagnose abnormal terminations.
+			switch (CrashContext.CrashType)
+			{
+				case ECrashContextType::Assert:           FDiagnosticLogger::Get().LogEvent(TEXT("Pipe/Assert"));          break;
+				case ECrashContextType::Ensure:           FDiagnosticLogger::Get().LogEvent(TEXT("Pipe/Ensure"));          break;
+				case ECrashContextType::Crash:            FDiagnosticLogger::Get().LogEvent(TEXT("Pipe/Crash"));           break;
+				case ECrashContextType::GPUCrash:         FDiagnosticLogger::Get().LogEvent(TEXT("Pipe/GPUCrash"));        break;
+				case ECrashContextType::Hang:             FDiagnosticLogger::Get().LogEvent(TEXT("Pipe/Hang"));            break;
+				case ECrashContextType::OutOfMemory:      FDiagnosticLogger::Get().LogEvent(TEXT("Pipe/OOM"));             break;
+				case ECrashContextType::AbnormalShutdown: FDiagnosticLogger::Get().LogEvent(TEXT("Pipe/AbnormalShutdown"));break;
+				default:                                  FDiagnosticLogger::Get().LogEvent(TEXT("Pipe/Unknown"));         break;
+			}
 		}
 
 		return SharedCtxIt == SharedCtxEndIt;
@@ -766,6 +940,8 @@ void RunCrashReportClient(const TCHAR* CommandLine)
 	GEngineLoop.PreInit(*FinalCommandLine);
 	check(GConfig && GConfig->IsReadyForUse());
 
+	FDiagnosticLogger::Get().LogEvent(TEXT("EngineInit/Start"));
+
 	// Make sure all UObject classes are registered and default properties have been initialized
 	ProcessNewlyLoadedUObjects();
 
@@ -786,6 +962,8 @@ void RunCrashReportClient(const TCHAR* CommandLine)
 
 	FPlatformErrorReport::Init();
 
+	FDiagnosticLogger::Get().LogEvent(TEXT("EngineInit/Done"));
+
 	if (MonitorPid == 0) // Does not monitor any process.
 	{
 		if (AnalyticsEnabledFromCmd)
@@ -805,20 +983,23 @@ void RunCrashReportClient(const TCHAR* CommandLine)
 	}
 	else // Launched in 'service mode - watches/serves a process'
 	{
+		FDiagnosticLogger::Get().LogEvent(TEXT("Monitor/Started"));
 		const int32 IdealFramerate = 30;
 		double LastTime = FPlatformTime::Seconds();
 		const float IdealFrameTime = 1.0f / IdealFramerate;
 
 		TSharedPtr<FRecoveryService> RecoveryServicePtr; // Note: Shared rather than Unique due to FRecoveryService only being a forward declaration in some builds
+
 #if CRASH_REPORT_WITH_RECOVERY
 		// Starts the disaster recovery service. This records transactions and allows users to recover from previous crashes.
 		RecoveryServicePtr = MakeShared<FRecoveryService>(MonitorPid);
+		FDiagnosticLogger::Get().LogEvent(TEXT("Recovery/Started"));
 #endif
 
 		// Try to open the process. The analytics shows that CRC sometime exits before the Editor. This can happen if OpenProcess() fails, maybe because
 		// this process was spawned in the Editor Pre-init phase and the OS did not make it available yet. In case of failure, give it a few trials.
 		FProcHandle MonitoredProcess;
-		for (int Trial = 0; Trial < 5; ++Trial)
+		for (int32 Trial = 0; Trial < 5; ++Trial)
 		{
 #if PLATFORM_WINDOWS
 			// We do not need to open a full access 
@@ -834,10 +1015,7 @@ void RunCrashReportClient(const TCHAR* CommandLine)
 			FPlatformProcess::Sleep(1); // Wait a seconds before retrying. Give time to the Editor to initialize.
 		}
 
-		if (!MonitoredProcess.IsValid())
-		{
-			UE_LOG(CrashReportClientLog, Error, TEXT("Failed to open monitor process handle!"));
-		}
+		FDiagnosticLogger::Get().LogEvent(MonitoredProcess.IsValid() ? TEXT("OpenProcess/Done") : TEXT("OpenProcess/Failed"));
 
 		auto GetProcessStatus = [](FProcHandle& ProcessHandle) -> TTuple<bool/*Running*/, TOptional<int32>/*ReturnCode*/>
 		{
@@ -873,6 +1051,8 @@ void RunCrashReportClient(const TCHAR* CommandLine)
 				FSharedCrashContext CrashContext;
 				if (IsCrashReportAvailable(MonitorPid, CrashContext, MonitorReadPipe))
 				{
+					FDiagnosticLogger::Get().LogEvent(TEXT("Report/Start"));
+
 					const bool bReportCrashAnalyticInfo = CrashContext.UserSettings.bSendUsageData;
 					if (bReportCrashAnalyticInfo)
 					{
@@ -901,6 +1081,8 @@ void RunCrashReportClient(const TCHAR* CommandLine)
 						}
 						FCrashReportAnalytics::Shutdown();
 					}
+
+					FDiagnosticLogger::Get().LogEvent(TEXT("Report/Done"));
 				}
 			}
 
@@ -931,7 +1113,7 @@ void RunCrashReportClient(const TCHAR* CommandLine)
 
 #if CRASH_REPORT_WITH_MTBF // Expected to be 1 when compiling CrashReportClientEditor.
 		{
-			UE_LOG(CrashReportClientLog, Log, TEXT("Starting Editor MTBF reporting"));
+			FDiagnosticLogger::Get().LogEvent(TEXT("MTBF/Start"));
 
 			// The loop above can exit before the Editor (monitored process) exits (because of IsEngineExitRequested()) if the user clicks 'Close Without Sending' very quickly, but for MTBF,
 			// it is desirable to have the Editor process return code. Give some extra time to the Editor to exit. If it doesn't exit within x seconds the next CRC instance will sent the
@@ -976,12 +1158,14 @@ void RunCrashReportClient(const TCHAR* CommandLine)
 				else
 				{
 					UE_LOG(CrashReportClientLog, Warning, TEXT("Failed to set Editor exit code. The session could not be found."));
+					FDiagnosticLogger::Get().LogEvent(TEXT("MTBF/NoSessionFound"));
 				}
 				FEditorAnalyticsSession::Unlock();
 			}
 			else
 			{
 				UE_LOG(CrashReportClientLog, Warning, TEXT("Failed to acquire lock to set exit code in Editor summary"));
+				FDiagnosticLogger::Get().LogEvent(TEXT("MTBF/LockSessionFail"));
 			}
 
 			if (bMonitoredProcessExited)
@@ -994,9 +1178,12 @@ void RunCrashReportClient(const TCHAR* CommandLine)
 					FCrashReportAnalytics::Initialize();
 					if (FCrashReportAnalytics::IsAvailable())
 					{
+						FDiagnosticLogger::Get().LogEvent(TEXT("MTFB/Done"));
 						// Send this session summary event (and the orphan ones if any).
 						FEditorSessionSummarySender EditorSessionSummarySender(FCrashReportAnalytics::GetProvider(), TEXT("CrashReportClient"), MonitorPid);
+						EditorSessionSummarySender.SetMonitorDiagnosticLogs(FDiagnosticLogger::LoadAllLogs());
 						EditorSessionSummarySender.Shutdown();
+						FDiagnosticLogger::ClearAllLogs(); // Logs (if any) were attached and sent if required and are not longer required.
 
 						// If the Editor thinks the session ended up abnormally, generate a crash report (to get the Editor logs and figure out why this happened).
 						if (bMonitoredSessionLoaded && TempCrashContext.UserSettings.bSendUnattendedBugReports)
@@ -1013,15 +1200,13 @@ void RunCrashReportClient(const TCHAR* CommandLine)
 				}
 				else
 				{
-					UE_LOG(CrashReportClientLog, Warning, TEXT("Failed to send Editor session summary event. Could not read the temporary file."));
+					FDiagnosticLogger::Get().LogEvent(TEXT("MTFB/NoTempCrash"));
 				}
 			}
 			else
 			{
-				UE_LOG(CrashReportClientLog, Warning, TEXT("Failed to send Editor session summary event. The Editor was still running."));
+				FDiagnosticLogger::Get().LogEvent(TEXT("MTFB/StillRunning"));
 			}
-
-			UE_LOG(CrashReportClientLog, Display, TEXT("Finished Editor MTBF reporting"));
 		}
 #endif
 		// Clean up the context file
