- Fixed the application title, getting it from the engine version rather than hardcoding it.

diff --git a/Engine/Source/Programs/CrashReportClient/Private/CrashReportClientApp.cpp b/Engine/Source/Programs/CrashReportClient/Private/CrashReportClientApp.cpp
--- a/Engine/Source/Programs/CrashReportClient/Private/CrashReportClientApp.cpp
+++ b/Engine/Source/Programs/CrashReportClient/Private/CrashReportClientApp.cpp
@@ -450,9 +450,13 @@ SubmitCrashReportResult RunWithUI(FPlatformErrorReport ErrorReport)
 	// open up the app window	
 	TSharedRef<SCrashReportClient> ClientControl = SNew(SCrashReportClient, CrashReportClient);
 
+	// Get the engine major version to display in title.
+	FBuildVersion BuildVersion;
+	uint16 MajorEngineVersion = FBuildVersion::TryRead(FBuildVersion::GetDefaultFileName(), BuildVersion) ? BuildVersion.GetEngineVersion().GetMajor() : 5;
+
 	TSharedRef<SWindow> Window = FSlateApplication::Get().AddWindow(
 		SNew(SWindow)
-		.Title(NSLOCTEXT("CrashReportClient", "CrashReportClientAppName", "Unreal Engine 4 Crash Reporter"))
+		.Title(FText::Format(NSLOCTEXT("CrashReportClient", "CrashReportClientAppName", "Unreal Engine {0} Crash Reporter"), MajorEngineVersion))
 		.HasCloseButton(FCrashReportCoreConfig::Get().IsAllowedToCloseWithoutSending())
 		.ClientSize(InitialWindowDimensions)
 		[
