Logged ensures file/line/condition/message captured by CrashReportClientEditor in the analytic diagnostic log.

diff --git a/Engine/Source/Programs/CrashReportClient/Private/CrashReportClientApp.cpp b/Engine/Source/Programs/CrashReportClient/Private/CrashReportClientApp.cpp
--- a/Engine/Source/Programs/CrashReportClient/Private/CrashReportClientApp.cpp
+++ b/Engine/Source/Programs/CrashReportClient/Private/CrashReportClientApp.cpp
@@ -982,11 +982,15 @@ void RunCrashReportClient(const TCHAR* CommandLine)
 					FDiagnosticLogger::Get().LogEvent(Result == SubmitCrashReportResult::SuccessDiscarded ? TEXT("Report/Discarded") :
 					                                 (Result == SubmitCrashReportResult::Failed ? TEXT("Report/Failed") : TEXT("Report/Sent")));
 
-					// Log the assert condition/file/line/message to the diagnostic log gathered by the analytics to enable grouping asserts later on.
+					// Log the assert and ensure condition/file/line/message to the diagnostic log gathered by the analytics to enable searching/grouping them later on.
 					if (CrashContext.CrashType == ECrashContextType::Assert)
 					{
 						FDiagnosticLogger::Get().LogEvent(FString::Printf(TEXT("Assert/Msg: %s"), CrashContext.ErrorMessage));
 					}
+					else if (CrashContext.CrashType == ECrashContextType::Ensure)
+					{
+						FDiagnosticLogger::Get().LogEvent(FString::Printf(TEXT("Ensure/Msg: %s"), CrashContext.ErrorMessage));
+					}
 
 					if (bReportCrashAnalyticInfo)
 					{
