Optmized how the Crash Context memory buffer used to serialize the context was used:
  - Prevented CrashReportClient from serializing the CrashContext to the memory buffer if it was already serialized once.
  - Increased the default space reserved by memory buffer used to serialize the crash context from 32K to 128K because serializing the crash context of 'debug crash' command line in Editor uses up to 112K.
  - Cleared the memory buffer before serializing the CrashContext in case it was serialized more than once preventing the internal buffer to grow needlessly. Also, the XML reader being limited, would only read the first one written, ignoring further and more recent ones appended.

diff --git a/Engine/Source/Programs/CrashReportClient/Private/CrashReportClientApp.cpp b/Engine/Source/Programs/CrashReportClient/Private/CrashReportClientApp.cpp
--- a/Engine/Source/Programs/CrashReportClient/Private/CrashReportClientApp.cpp
+++ b/Engine/Source/Programs/CrashReportClient/Private/CrashReportClientApp.cpp
@@ -468,11 +468,7 @@ FPlatformErrorReport CollectErrorReport(FRecoveryService* RecoveryService, uint3
 		}
 	}
 
-	// Setup the FPrimaryCrashProperties singleton. If the path is not set it is most likely
-	// that we have crashed during static init, in which case we need to construct a directory
-	// ourself.
-	CrashContext.SerializeContentToBuffer();
-
+	// If the path is not set it is most likely that we have crashed during static init, in which case we need to construct a directory ourself.
 	FString ReportDirectoryAbsolutePath(SharedCrashContext.CrashFilesDirectory);
 	bool DirectoryExists = true;
 	if (ReportDirectoryAbsolutePath.IsEmpty())
@@ -509,8 +505,15 @@ FPlatformErrorReport CollectErrorReport(FRecoveryService* RecoveryService, uint3
 	}
 #endif
 
-	const TCHAR* CrachContextBuffer = *CrashContext.GetBuffer();
-	FPrimaryCrashProperties::Set(new FCrashContext(ReportDirectoryAbsolutePath / TEXT("CrashContext.runtime-xml"), CrachContextBuffer));
+	// If the crash context wasn't implicitely serialized by SerializeAsXML() above, serialize it now.
+	if (CrashContext.GetBuffer().IsEmpty())
+	{
+		CrashContext.SerializeContentToBuffer();
+	}
+
+	// Setup the FPrimaryCrashProperties singleton.
+	const TCHAR* CrashContextBuffer = *CrashContext.GetBuffer();
+	FPrimaryCrashProperties::Set(new FCrashContext(ReportDirectoryAbsolutePath / TEXT("CrashContext.runtime-xml"), CrashContextBuffer));
 
 	FPlatformErrorReport ErrorReport(ReportDirectoryAbsolutePath);
 
diff --git a/Engine/Source/Runtime/Core/Private/GenericPlatform/GenericPlatformCrashContext.cpp b/Engine/Source/Runtime/Core/Private/GenericPlatform/GenericPlatformCrashContext.cpp
--- a/Engine/Source/Runtime/Core/Private/GenericPlatform/GenericPlatformCrashContext.cpp
+++ b/Engine/Source/Runtime/Core/Private/GenericPlatform/GenericPlatformCrashContext.cpp
@@ -431,7 +431,7 @@ FGenericCrashContext::FGenericCrashContext(ECrashContextType InType, const TCHAR
 	, ErrorMessage(InErrorMessage)
 	, NumMinidumpFramesToIgnore(0)
 {
-	CommonBuffer.Reserve( 32768 );
+	CommonBuffer.Reserve( 128 * 1024 ); // NOTE: The Editor command 'debug crash' uses about 112K characters on Windows when the crash is serialized by SerializeContentToBuffer()
 	CrashContextIndex = StaticCrashContextIndex++;
 }
 
@@ -601,6 +601,10 @@ void FGenericCrashContext::SerializeUserSettings(FString& Buffer)
 
 void FGenericCrashContext::SerializeContentToBuffer() const
 {
+	// Clear the buffer in case the content is serialized more than once, keeping the most up to date values and preventing to store several XML documents in the same buffer.
+	// When the buffer is passed to our XML reader, only the first document <FGenericCrashContext></<FGenericCrashContext> is read and further ones (most recent ones) are ignored.
+	CommonBuffer.Reset();
+
 	TCHAR CrashGUID[CrashGUIDLength];
 	GetUniqueCrashName(CrashGUID, CrashGUIDLength);
 
