Fix the task tag scope of the crash report client on the Mac.

#rnx

diff --git a/Engine/Source/Programs/CrashReportClient/Private/CrashReportClientApp.cpp b/Engine/Source/Programs/CrashReportClient/Private/CrashReportClientApp.cpp
--- a/Engine/Source/Programs/CrashReportClient/Private/CrashReportClientApp.cpp
+++ b/Engine/Source/Programs/CrashReportClient/Private/CrashReportClientApp.cpp
@@ -885,7 +885,9 @@ static bool WasAbnormalShutdown(const FEditorAnalyticsSession& AnalyticSession)
 
 void RunCrashReportClient(const TCHAR* CommandLine)
 {
+#if !PLATFORM_MAC
 	FTaskTagScope ThreadScope(ETaskTag::EGameThread); // Main thread is the game thread.
+#endif
 
 	FDiagnosticLogger::Get().LogEvent(FString::Printf(TEXT("CRC/Init:%s"), *FDateTime::UtcNow().ToString()));
 
