Added a diagnostic log to CRC when the handle returned by OpenProcess() is invalid. This handle is used to stack walk the crash and generate a minidump.

diff --git a/Engine/Source/Programs/CrashReportClient/Private/CrashReportClientApp.cpp b/Engine/Source/Programs/CrashReportClient/Private/CrashReportClientApp.cpp
--- a/Engine/Source/Programs/CrashReportClient/Private/CrashReportClientApp.cpp
+++ b/Engine/Source/Programs/CrashReportClient/Private/CrashReportClientApp.cpp
@@ -439,6 +439,10 @@ FPlatformErrorReport CollectErrorReport(FRecoveryService* RecoveryService, uint3
 
 	// @note: This API is only partially implemented on Mac OS and Linux.
 	FProcHandle ProcessHandle = FPlatformProcess::OpenProcess(Pid);
+	if (!ProcessHandle.IsValid())
+	{
+		FDiagnosticLogger::Get().LogEvent(TEXT("Report/OpenProcessFail"));
+	}
 
 	// First init the static crash context state
 	FPlatformCrashContext::InitializeFromContext(
@@ -1089,7 +1093,7 @@ void RunCrashReportClient(const TCHAR* CommandLine)
 					// that the timeout in the crashing process waiting for CRC to reply 'continue' is too short.
 					if (!bCrashedThreadCallstackAvailable)
 					{
-						FDiagnosticLogger::Get().LogEvent(FString::Printf(TEXT("Report/NoPCallstack")));
+						FDiagnosticLogger::Get().LogEvent(TEXT("Report/NoPCallstack"));
 					}
 
 #if CRASH_REPORT_WITH_RECOVERY
