Fixed CrashReportClientEditor compile error in unity builds.

diff --git a/Engine/Source/Programs/CrashReportClient/Private/CrashReportAnalyticsSessionSummary.cpp b/Engine/Source/Programs/CrashReportClient/Private/CrashReportAnalyticsSessionSummary.cpp
--- a/Engine/Source/Programs/CrashReportClient/Private/CrashReportAnalyticsSessionSummary.cpp
+++ b/Engine/Source/Programs/CrashReportClient/Private/CrashReportAnalyticsSessionSummary.cpp
@@ -37,6 +37,8 @@ LRESULT CALLBACK CrashReportAnalyticsSessionSummaryWindowProc(HWND Hwnd, UINT Ms
 	return DefWindowProc(Hwnd, Msg, wParam, lParam);
 }
 
+#include "Windows/HideWindowsPlatformTypes.h"
+
 #endif
 
 namespace CrcAnalyticsProperties
