Fixed Editor analytics session summary key 'MonitorLogs' as 'MonitorLog' to be consistent with previous versions.

diff --git a/Engine/Source/Programs/CrashReportClient/Private/CrashReportAnalyticsSessionSummary.cpp b/Engine/Source/Programs/CrashReportClient/Private/CrashReportAnalyticsSessionSummary.cpp
--- a/Engine/Source/Programs/CrashReportClient/Private/CrashReportAnalyticsSessionSummary.cpp
+++ b/Engine/Source/Programs/CrashReportClient/Private/CrashReportAnalyticsSessionSummary.cpp
@@ -57,7 +57,7 @@ namespace CrcAnalyticsProperties
 	/** If CRC is about to close because the system sent a quit signal. */
 	static const TAnalyticsProperty<bool> QuitSignalRecv(TEXT("MonitorQuitSignalRecv"));
 	/** The CRC diagnostic logs. */
-	static const TAnalyticsProperty<FString> DiagnosticLogs(TEXT("MonitorLogs"));
+	static const TAnalyticsProperty<FString> DiagnosticLogs(TEXT("MonitorLog"));
 	/** The CRC session duration in seconds. */
 	static const TAnalyticsProperty<int32> SessionDurationSecs(TEXT("MonitorSessionDuration"));
 	/** The battery level, if known. */
