Removed duplicated module dependency on Analytics from CrashReportCLientEditor build script.
Added missing include in AnalyticsPropertyStore.cpp
Added ENGINE_API to expose EngineAnlytixSessionSummary publicly.
Added the longuest unattended crash report processing measured during a session to CRC summary to track how fast or how slow CRC is to manage ensures/stalls.

diff --git a/Engine/Source/Programs/CrashReportClient/CrashReportClientEditor.Build.cs b/Engine/Source/Programs/CrashReportClient/CrashReportClientEditor.Build.cs
--- a/Engine/Source/Programs/CrashReportClient/CrashReportClientEditor.Build.cs
+++ b/Engine/Source/Programs/CrashReportClient/CrashReportClientEditor.Build.cs
@@ -19,7 +19,6 @@ public class CrashReportClientEditor : CrashReportClient
 		PrivateDependencyModuleNames.AddRange(
 			new string[]
 			{
-				"Analytics",
 				"EditorAnalyticsSession",
 			}
 		);
diff --git a/Engine/Source/Programs/CrashReportClient/Private/CrashReportAnalyticsSessionSummary.cpp b/Engine/Source/Programs/CrashReportClient/Private/CrashReportAnalyticsSessionSummary.cpp
--- a/Engine/Source/Programs/CrashReportClient/Private/CrashReportAnalyticsSessionSummary.cpp
+++ b/Engine/Source/Programs/CrashReportClient/Private/CrashReportAnalyticsSessionSummary.cpp
@@ -44,8 +44,8 @@ LRESULT CALLBACK CrashReportAnalyticsSessionSummaryWindowProc(HWND Hwnd, UINT Ms
 
 namespace CrcAnalyticsProperties
 {
-	//NOTE: Update this when you add/remove/change key behavior.
-	constexpr uint32 CrcAnalyticsSummaryVersion = 1;
+	//NOTE: Update this when you add/remove/change key behavior. That's useful to track how one changes affects metrics in-dev where users don't always have an engine versions.
+	constexpr uint32 CrcAnalyticsSummaryVersion = 2;
 
 	/** The exit code of the monitored application. */
 	static const TAnalyticsProperty<int32> MonitoredAppExitCode(TEXT("ExitCode"));
@@ -90,6 +90,8 @@ namespace CrcAnalyticsProperties
 	static const TAnalyticsProperty<uint32> EnsureCount(TEXT("MonitorEnsureCount"));
 	/** Number of assert handled by CRC.*/
 	static const TAnalyticsProperty<uint32> AssertCount(TEXT("MonitorAssertCount"));
+	/** The worst unattended report time measured (if any). The user is not involved, so it measure how fast CRC can process a crash, especially ensures and stalls. */
+	static const TAnalyticsProperty<float> LonguestUnattendedReportSecs(TEXT("MonitorLongestUnattendedReportSecs"));
 }
 
 namespace CrashReportClientUtils
@@ -301,15 +303,13 @@ void FCrashReportAnalyticsSessionSummary::Initialize(const FString& ProcessGroup
 						LogEvent(TEXT("CRC/OpenProcessFailed"));
 						return;
 					}
-					else
-					{
-						LogEvent(TEXT("CRC/Monitoring"));
-					}
 
 					CrashReportClientUtils::InitPlatformSpecific();
 					double NextFlushTimeSecs = FPlatformTime::Seconds();
 					bool bFlushedLowBattery = false;
 
+					LogEvent(TEXT("CRC/Monitoring")); // About to enter the loop.
+
 					while (!bShutdown)
 					{
 						CrashReportClientUtils::TickPlatformSpecific();
@@ -599,6 +599,7 @@ void FCrashReportAnalyticsSessionSummary::OnCrashReportProcessing(bool bUserInte
 	CrcAnalyticsProperties::IsCollectingCrash.Set(PropertyStore.Get(), false);
 	CrcAnalyticsProperties::IsProcessingCrash.Set(PropertyStore.Get(), true);
 	CrashReportProcessingStartTimeSecs = FPlatformTime::Seconds();
+	bProcessingCrashUnattended = !bUserInteractive;
 	FCrashReportAnalyticsSessionSummary::Get().LogEvent(FString::Printf(TEXT("Report/Process:%s"), bUserInteractive ? TEXT("Interactive") : TEXT("Unattended")));
 }
 
@@ -612,6 +613,13 @@ void FCrashReportAnalyticsSessionSummary::OnCrashReportCompleted(bool bSubmitted
 	double CollectTimeSecs = CrashReportProcessingStartTimeSecs - CrashReportCollectingStartTimeSecs;
 	double ProcessTimeSecs = CurrTimeSecs - CrashReportProcessingStartTimeSecs;
 	double TotalTimeSecs = CurrTimeSecs - CrashReportStartTimeSecs;
+
+	if (bProcessingCrashUnattended) // No UI shown to the user that could amplify the time.
+	{
+		CrcAnalyticsProperties::LonguestUnattendedReportSecs.Set(PropertyStore.Get(), static_cast<float>(TotalTimeSecs), [](const float* Actual, const float& Proposed) { return Actual == nullptr || Proposed > *Actual; });
+	}
+	bProcessingCrashUnattended = false;
+
 	const TCHAR* Event = bSubmitted ? TEXT("Report/Sent") : TEXT("Report/Discarded");
 	LogEvent(FString::Printf(TEXT("%s:Collect=%.1fs:Process=%.1fs:Total=%.1fs"), Event, CollectTimeSecs, ProcessTimeSecs, TotalTimeSecs));
 }
diff --git a/Engine/Source/Programs/CrashReportClient/Private/CrashReportAnalyticsSessionSummary.h b/Engine/Source/Programs/CrashReportClient/Private/CrashReportAnalyticsSessionSummary.h
--- a/Engine/Source/Programs/CrashReportClient/Private/CrashReportAnalyticsSessionSummary.h
+++ b/Engine/Source/Programs/CrashReportClient/Private/CrashReportAnalyticsSessionSummary.h
@@ -120,6 +120,9 @@ private:
 	/** Time at which processing the last crash started. */
 	double CrashReportProcessingStartTimeSecs = 0.0;
 
+	/** Whether the crash being processed unattended. */
+	bool bProcessingCrashUnattended = false;
+
 	/** Prevent a reentrency in the diagnostic logger. */
 	bool bLoggerReentrantGuard = false;
 
diff --git a/Engine/Source/Runtime/Analytics/Analytics/Private/AnalyticsPropertyStore.cpp b/Engine/Source/Runtime/Analytics/Analytics/Private/AnalyticsPropertyStore.cpp
--- a/Engine/Source/Runtime/Analytics/Analytics/Private/AnalyticsPropertyStore.cpp
+++ b/Engine/Source/Runtime/Analytics/Analytics/Private/AnalyticsPropertyStore.cpp
@@ -8,6 +8,7 @@
 #include "HAL/PlatformTime.h"
 #include "Misc/SecureHash.h"
 #include "Misc/ScopeLock.h"
+#include "Misc/ScopeExit.h"
 #include "ProfilingDebugging/CpuProfilerTrace.h"
 
 namespace AnalyticsPropertyStoreUtils
diff --git a/Engine/Source/Runtime/Engine/Public/Analytics/EngineAnalyticsSessionSummary.h b/Engine/Source/Runtime/Engine/Public/Analytics/EngineAnalyticsSessionSummary.h
--- a/Engine/Source/Runtime/Engine/Public/Analytics/EngineAnalyticsSessionSummary.h
+++ b/Engine/Source/Runtime/Engine/Public/Analytics/EngineAnalyticsSessionSummary.h
@@ -14,7 +14,7 @@ struct FUserActivity;
 /**
  * Collects engine events/stats and stores a summary on disk until reported by the analytics session summary manager.
  */
-class FEngineAnalyticsSessionSummary
+class ENGINE_API FEngineAnalyticsSessionSummary
 {
 public:
 	/**
