Fix non-unity build

#rnx

diff --git a/Engine/Source/Programs/CrashReportClient/Private/CrashReportAnalyticsSessionSummary.cpp b/Engine/Source/Programs/CrashReportClient/Private/CrashReportAnalyticsSessionSummary.cpp
--- a/Engine/Source/Programs/CrashReportClient/Private/CrashReportAnalyticsSessionSummary.cpp
+++ b/Engine/Source/Programs/CrashReportClient/Private/CrashReportAnalyticsSessionSummary.cpp
@@ -14,6 +14,8 @@
 #include "Misc/EngineVersion.h"
 #include "Misc/Paths.h"
 #include "Misc/ScopeLock.h"
+#include "Misc/CoreDelegates.h"
+#include "Misc/OutputDeviceRedirector.h"
 #include "Templates/UnrealTemplate.h"
 
 DEFINE_LOG_CATEGORY_STATIC(LogCrashReportClientDiagnostics, Log, All)
