Stripped the callstack from the ensure error message when logging the ensure message in the analytics diagnostic log.
  - This information is not required for analytic purpose.

diff --git a/Engine/Source/Programs/CrashReportClient/Private/CrashReportAnalyticsSessionSummary.cpp b/Engine/Source/Programs/CrashReportClient/Private/CrashReportAnalyticsSessionSummary.cpp
--- a/Engine/Source/Programs/CrashReportClient/Private/CrashReportAnalyticsSessionSummary.cpp
+++ b/Engine/Source/Programs/CrashReportClient/Private/CrashReportAnalyticsSessionSummary.cpp
@@ -11,6 +11,7 @@
 #include "GenericPlatform/GenericPlatformCrashContext.h"
 #include "HAL/PlatformTime.h"
 #include "HAL/Thread.h"
+#include "Internationalization/Regex.h"
 #include "Logging/LogMacros.h"
 #include "Misc/EngineVersion.h"
 #include "Misc/Paths.h"
@@ -49,7 +50,8 @@ namespace CrcAnalyticsProperties
 {
 	// NOTE: Update this when you add/remove/change key behavior. That's useful to track how one changes affects metrics in-dev where users don't always have an engine versions.
 	//     - V3 -> Windows optimization for stall/ensure -> The engine only captures the responsible thread so CRC walks 1 thread rather than all threads.
-	constexpr uint32 CrcAnalyticsSummaryVersion = 3;
+	//     - V4 -> Stripped ensure callstack from the ensure error message to remove noise in the diagnostic log.
+	constexpr uint32 CrcAnalyticsSummaryVersion = 4;
 
 	/** The exit code of the monitored application. */
 	static const TAnalyticsProperty<int32> MonitoredAppExitCode(TEXT("ExitCode"));
@@ -586,7 +588,18 @@ void FCrashReportAnalyticsSessionSummary::OnCrashReportStarted(ECrashContextType
 	else if (CrashType == ECrashContextType::Ensure)
 	{
 		CrcAnalyticsProperties::EnsureCount.Update(PropertyStore.Get(), [](uint32& Actual) { ++Actual; return true; });
-		LogEvent(FString::Printf(TEXT("Ensure/Msg: %s"), ErrorMsg));
+
+		// Ensure messages include the ensure call stack. That's not useful for analytics, try keeping the essential only (ensure condition, file, line)
+		FRegexPattern Pattern(TEXT(R"(.*\[File:.*\]\s*\[Line:\s\d+\])")); // Need help with regex? Try https://regex101.com/
+		FRegexMatcher Matcher(Pattern, ErrorMsg);
+		if (Matcher.FindNext())
+		{
+			LogEvent(FString::Printf(TEXT("Ensure/Msg: %s"), *Matcher.GetCaptureGroup(0)));
+		}
+		else
+		{
+			LogEvent(FString::Printf(TEXT("Ensure/Msg: %s"), ErrorMsg));
+		}
 	}
 }
 
