Added analytics to report more granular metrics about CRC performance when handling crash/ensure/stall.

diff --git a/Engine/Source/Programs/CrashReportClient/Private/CrashReportAnalyticsSessionSummary.cpp b/Engine/Source/Programs/CrashReportClient/Private/CrashReportAnalyticsSessionSummary.cpp
--- a/Engine/Source/Programs/CrashReportClient/Private/CrashReportAnalyticsSessionSummary.cpp
+++ b/Engine/Source/Programs/CrashReportClient/Private/CrashReportAnalyticsSessionSummary.cpp
@@ -51,6 +51,7 @@ namespace CrcAnalyticsProperties
 	// NOTE: Update this when you add/remove/change key behavior. That's useful to track how one changes affects metrics in-dev where users don't always have an engine versions.
 	//     - V3 -> Windows optimization for stall/ensure -> The engine only captures the responsible thread so CRC walks 1 thread rather than all threads.
 	//     - V4 -> Stripped ensure callstack from the ensure error message to remove noise in the diagnostic log.
+	//     - V5 -> Measured time to stack-walk, gather files and addded stall count.
 	constexpr uint32 CrcAnalyticsSummaryVersion = 4;
 
 	/** The exit code of the monitored application. */
@@ -96,6 +97,8 @@ namespace CrcAnalyticsProperties
 	static const TAnalyticsProperty<uint32> EnsureCount(TEXT("MonitorEnsureCount"));
 	/** Number of assert handled by CRC.*/
 	static const TAnalyticsProperty<uint32> AssertCount(TEXT("MonitorAssertCount"));
+	/** Number of stalls handed by CRC. */
+	static const TAnalyticsProperty<uint32> StallCount(TEXT("MonitorStallCount"));
 	/** The worst unattended report time measured (if any). The user is not involved, so it measure how fast CRC can process a crash, especially ensures and stalls. */
 	static const TAnalyticsProperty<float> LonguestUnattendedReportSecs(TEXT("MonitorLongestUnattendedReportSecs"));
 }
@@ -293,6 +296,7 @@ void FCrashReportAnalyticsSessionSummary::Initialize(const FString& ProcessGroup
 				CrcAnalyticsProperties::ReportCount.Set(PropertyStore.Get(), 0);
 				CrcAnalyticsProperties::EnsureCount.Set(PropertyStore.Get(), 0);
 				CrcAnalyticsProperties::AssertCount.Set(PropertyStore.Get(), 0);
+				CrcAnalyticsProperties::StallCount.Set(PropertyStore.Get(), 0);
 
 				UpdatePowerStatus();
 				Flush();
@@ -575,6 +579,13 @@ void FCrashReportAnalyticsSessionSummary::OnCrashReportStarted(ECrashContextType
 {
 	CrashReportStartTimeSecs = FPlatformTime::Seconds();
 
+	// Reset the timers for the current crash report (a safety measure).
+	CrashReportCollectingStartTimeSecs = CrashReportStartTimeSecs;
+	CrashReportStackWalkingStartTimeSecs = CrashReportStartTimeSecs;
+	CrashReportGatheringFilesStartTimeSecs = CrashReportStartTimeSecs;
+	CrashReportSignalingRemoteAppTimeSecs = CrashReportStartTimeSecs;
+	CrashReportProcessingStartTimeSecs = CrashReportStartTimeSecs;
+
 	CrcAnalyticsProperties::IsReportingCrash.Set(PropertyStore.Get(), true);
 	CrcAnalyticsProperties::ReportCount.Update(PropertyStore.Get(), [](uint32& Actual) { ++Actual; return true; });
 	LogEvent(FString::Printf(TEXT("Report/Start:%s"), *FDateTime::UtcNow().ToString()));
@@ -601,6 +612,10 @@ void FCrashReportAnalyticsSessionSummary::OnCrashReportStarted(ECrashContextType
 			LogEvent(FString::Printf(TEXT("Ensure/Msg: %s"), ErrorMsg));
 		}
 	}
+	else if (CrashType == ECrashContextType::Stall)
+	{
+		CrcAnalyticsProperties::StallCount.Update(PropertyStore.Get(), [](uint32& Actual) { ++Actual; return true; });
+	}
 }
 
 void FCrashReportAnalyticsSessionSummary::OnCrashReportCollecting()
@@ -611,6 +626,21 @@ void FCrashReportAnalyticsSessionSummary::OnCrashReportCollecting()
 	FCrashReportAnalyticsSessionSummary::Get().LogEvent(TEXT("Report/Collect"));
 }
 
+void FCrashReportAnalyticsSessionSummary::OnCrashReportRemoteStackWalking()
+{
+	CrashReportStackWalkingStartTimeSecs = FPlatformTime::Seconds();
+}
+
+void FCrashReportAnalyticsSessionSummary::OnCrashReportGatheringFiles()
+{
+	CrashReportGatheringFilesStartTimeSecs = FPlatformTime::Seconds();
+}
+
+void FCrashReportAnalyticsSessionSummary::OnCrashReportSignalingAppToResume()
+{
+	CrashReportSignalingRemoteAppTimeSecs = FPlatformTime::Seconds();
+}
+
 void FCrashReportAnalyticsSessionSummary::OnCrashReportProcessing(bool bUserInteractive)
 {
 	CrcAnalyticsProperties::IsCollectingCrash.Set(PropertyStore.Get(), false);
@@ -627,8 +657,20 @@ void FCrashReportAnalyticsSessionSummary::OnCrashReportCompleted(bool bSubmitted
 	CrcAnalyticsProperties::IsReportingCrash.Set(PropertyStore.Get(), false);
 
 	double CurrTimeSecs = FPlatformTime::Seconds();
+
+	// Total time required to remote stack walk, gather files, respond to the monited app.
 	double CollectTimeSecs = CrashReportProcessingStartTimeSecs - CrashReportCollectingStartTimeSecs;
+
+	// Total time required to remote stack walk
+	double StackWalkSecs = CrashReportGatheringFilesStartTimeSecs - CrashReportStackWalkingStartTimeSecs;
+
+	// Total time required to gather files (copy log and generates minidump).
+	double GatherFileSecs = CrashReportSignalingRemoteAppTimeSecs - CrashReportGatheringFilesStartTimeSecs;
+
+	// Total time required by CRC to process the crash report (resolve symbols + showing the UI + user time to respond if the report is interactive).
 	double ProcessTimeSecs = CurrTimeSecs - CrashReportProcessingStartTimeSecs;
+
+	// Total time CRC main thread was used to process the crash.
 	double TotalTimeSecs = CurrTimeSecs - CrashReportStartTimeSecs;
 
 	if (bProcessingCrashUnattended) // No UI shown to the user that could amplify the time.
@@ -638,7 +680,7 @@ void FCrashReportAnalyticsSessionSummary::OnCrashReportCompleted(bool bSubmitted
 	bProcessingCrashUnattended = false;
 
 	const TCHAR* Event = bSubmitted ? TEXT("Report/Sent") : TEXT("Report/Discarded");
-	LogEvent(FString::Printf(TEXT("%s:Collect=%.1fs:Process=%.1fs:Total=%.1fs"), Event, CollectTimeSecs, ProcessTimeSecs, TotalTimeSecs));
+	LogEvent(FString::Printf(TEXT("%s:Walk=%.1fs:Gather=%.1fs:Collect=%.1fs:Process=%.1fs:Total=%.1fs"), Event, StackWalkSecs, GatherFileSecs, CollectTimeSecs, ProcessTimeSecs, TotalTimeSecs));
 }
 
 bool FCrashReportAnalyticsSessionSummary::UpdatePowerStatus()
diff --git a/Engine/Source/Programs/CrashReportClient/Private/CrashReportAnalyticsSessionSummary.h b/Engine/Source/Programs/CrashReportClient/Private/CrashReportAnalyticsSessionSummary.h
--- a/Engine/Source/Programs/CrashReportClient/Private/CrashReportAnalyticsSessionSummary.h
+++ b/Engine/Source/Programs/CrashReportClient/Private/CrashReportAnalyticsSessionSummary.h
@@ -61,6 +61,15 @@ public:
 	/** Invoked when CRC starts collecting the crash artifacts. */
 	void OnCrashReportCollecting();
 
+	/** Invoked when CRC starts stack-walking the crash. */
+	void OnCrashReportRemoteStackWalking();
+
+	/** Invoked when CRC starts collecting crash artifacts such as the log/minidump/etc. */
+	void OnCrashReportGatheringFiles();
+
+	/** Invoked when CRC finished collecting crash artifacts and replies to the monitored app. */
+	void OnCrashReportSignalingAppToResume();
+
 	/** Invoked when CRC starts the process to submit the crash report (this also includes showing the dialog to user in needed. */
 	void OnCrashReportProcessing(bool bUserInteractive);
 
@@ -117,6 +126,15 @@ private:
 	/** Time at which collecting the artifacts of the last crash started. */
 	double CrashReportCollectingStartTimeSecs = 0.0;
 
+	/** Time at which remote stack-walking started. */
+	double CrashReportStackWalkingStartTimeSecs = 0.0;
+
+	/** Time at which gathering crash files (logs, minidump, etct) started. */
+	double CrashReportGatheringFilesStartTimeSecs = 0.0;
+
+	/** Time at which CRC signaled the application to resume. */
+	double CrashReportSignalingRemoteAppTimeSecs = 0.0;
+
 	/** Time at which processing the last crash started. */
 	double CrashReportProcessingStartTimeSecs = 0.0;
 
diff --git a/Engine/Source/Programs/CrashReportClient/Private/CrashReportClientApp.cpp b/Engine/Source/Programs/CrashReportClient/Private/CrashReportClientApp.cpp
--- a/Engine/Source/Programs/CrashReportClient/Private/CrashReportClientApp.cpp
+++ b/Engine/Source/Programs/CrashReportClient/Private/CrashReportClientApp.cpp
@@ -464,6 +464,8 @@ FPlatformErrorReport CollectErrorReport(FRecoveryService* RecoveryService, uint3
 	CrashContext.SetCrashedThreadId(SharedCrashContext.CrashingThreadId);
 	CrashContext.SetNumMinidumpFramesToIgnore(0);
 
+	FCrashReportAnalyticsSessionSummary::Get().OnCrashReportRemoteStackWalking();
+
 	// Initialize the stack walking for the monitored process (effectively overriding this process stack walking functionality)
 	FPlatformStackWalk::InitStackWalkingForProcess(ProcessHandle);
 
@@ -555,6 +557,8 @@ FPlatformErrorReport CollectErrorReport(FRecoveryService* RecoveryService, uint3
 		}
 	}
 
+	FCrashReportAnalyticsSessionSummary::Get().OnCrashReportGatheringFiles();
+
 	// If the path is not set it is most likely that we have crashed during static init, in which case we need to construct a directory ourself.
 	FString ReportDirectoryAbsolutePath(SharedCrashContext.CrashFilesDirectory);
 	bool DirectoryExists = true;
@@ -572,6 +576,8 @@ FPlatformErrorReport CollectErrorReport(FRecoveryService* RecoveryService, uint3
 		CrashContext.CopyPlatformSpecificFiles(*ReportDirectoryAbsolutePath, SharedCrashContext.PlatformCrashContext);
 	}
 
+	FCrashReportAnalyticsSessionSummary::Get().OnCrashReportSignalingAppToResume();
+
 	// At this point the game can continue execution. It is important this happens
 	// as soon as thread state and minidump has been created, so that ensures cause
 	// as little hitch as possible.
