Updated the crash report analytic session summary version number from 4 to 5.

diff --git a/Engine/Source/Programs/CrashReportClient/Private/CrashReportAnalyticsSessionSummary.cpp b/Engine/Source/Programs/CrashReportClient/Private/CrashReportAnalyticsSessionSummary.cpp
--- a/Engine/Source/Programs/CrashReportClient/Private/CrashReportAnalyticsSessionSummary.cpp
+++ b/Engine/Source/Programs/CrashReportClient/Private/CrashReportAnalyticsSessionSummary.cpp
@@ -52,7 +52,7 @@ namespace CrcAnalyticsProperties
 	//     - V3 -> Windows optimization for stall/ensure -> The engine only captures the responsible thread so CRC walks 1 thread rather than all threads.
 	//     - V4 -> Stripped ensure callstack from the ensure error message to remove noise in the diagnostic log.
 	//     - V5 -> Measured time to stack-walk, gather files and addded stall count.
-	constexpr uint32 CrcAnalyticsSummaryVersion = 4;
+	constexpr uint32 CrcAnalyticsSummaryVersion = 5;
 
 	/** The exit code of the monitored application. */
 	static const TAnalyticsProperty<int32> MonitoredAppExitCode(TEXT("ExitCode"));
