Added temporary diagnostic code to CrashReportClient in hopt to narrow down why it suspiciouly die often.

diff --git a/Engine/Source/Programs/CrashReportClient/Private/CrashReportAnalyticsSessionSummary.cpp b/Engine/Source/Programs/CrashReportClient/Private/CrashReportAnalyticsSessionSummary.cpp
--- a/Engine/Source/Programs/CrashReportClient/Private/CrashReportAnalyticsSessionSummary.cpp
+++ b/Engine/Source/Programs/CrashReportClient/Private/CrashReportAnalyticsSessionSummary.cpp
@@ -66,7 +66,8 @@ namespace CrcAnalyticsProperties
 	//     - V3 -> Windows optimization for stall/ensure -> The engine only captures the responsible thread so CRC walks 1 thread rather than all threads.
 	//     - V4 -> Stripped ensure callstack from the ensure error message to remove noise in the diagnostic log.
 	//     - V5 -> Measured time to stack-walk, gather files and addded stall count.
-	constexpr uint32 CrcAnalyticsSummaryVersion = 5;
+	//     - V6 -> Added MonitorTickCount, MonitorQueryingPipe and App/Death log.
+	constexpr uint32 CrcAnalyticsSummaryVersion = 6;
 
 	/** The exit code of the monitored application. */
 	static const TAnalyticsProperty<int32> MonitoredAppExitCode(TEXT("ExitCode"));
@@ -81,6 +82,8 @@ namespace CrcAnalyticsProperties
 	static const TAnalyticsProperty<FDateTime> StartupTimestamp(TEXT("MonitorStartupTimestamp"));
 	/** The CRC timestamp. */
 	static const TAnalyticsProperty<FDateTime> Timestamp(TEXT("MonitorTimestamp"));
+	/** Number of time CRC analytic thread ticked. */
+	static const TAnalyticsProperty<uint32> TickCount(TEXT("MonitorTickCount"));
 	/** If CRC raised an exception that was captured by SEH, this is the exception code. */
 	static const TAnalyticsProperty<int32> ExceptCode(TEXT("MonitorExceptCode"));
 	/** If CRC is about to close because the system sent a quit signal. */
@@ -115,6 +118,8 @@ namespace CrcAnalyticsProperties
 	static const TAnalyticsProperty<uint32> StallCount(TEXT("MonitorStallCount"));
 	/** The worst unattended report time measured (if any). The user is not involved, so it measure how fast CRC can process a crash, especially ensures and stalls. */
 	static const TAnalyticsProperty<float> LonguestUnattendedReportSecs(TEXT("MonitorLongestUnattendedReportSecs"));
+	/** Flags raised when the CRC is reading the pipe. */
+	static const TAnalyticsProperty<bool> ReadingPipeForCrash(TEXT("MonitorQueryingPipe"));
 }
 
 namespace CrashReportClientUtils
@@ -298,6 +303,7 @@ void FCrashReportAnalyticsSessionSummary::Initialize(const FString& ProcessGroup
 				CrcAnalyticsProperties::SummaryVersionNumber.Set(PropertyStore.Get(), CrcAnalyticsProperties::CrcAnalyticsSummaryVersion);
 				CrcAnalyticsProperties::StartupTimestamp.Set(PropertyStore.Get(), FDateTime::UtcNow());
 				CrcAnalyticsProperties::Timestamp.Set(PropertyStore.Get(), FDateTime::UtcNow());
+				CrcAnalyticsProperties::TickCount.Set(PropertyStore.Get(), 0);
 				CrcAnalyticsProperties::SessionDurationSecs.Set(PropertyStore.Get(), FMath::FloorToInt(static_cast<float>(FPlatformTime::Seconds() - SessionStartTimeSecs)));
 				CrcAnalyticsProperties::DiagnosticLogs.Set(PropertyStore.Get(), DiagnosticLog, CrashReportClientUtils::MaxDiagnosticLogLen);
 				CrcAnalyticsProperties::IsReportingCrash.Set(PropertyStore.Get(), false);
@@ -311,6 +317,7 @@ void FCrashReportAnalyticsSessionSummary::Initialize(const FString& ProcessGroup
 				CrcAnalyticsProperties::EnsureCount.Set(PropertyStore.Get(), 0);
 				CrcAnalyticsProperties::AssertCount.Set(PropertyStore.Get(), 0);
 				CrcAnalyticsProperties::StallCount.Set(PropertyStore.Get(), 0);
+				CrcAnalyticsProperties::ReadingPipeForCrash.Set(PropertyStore.Get(), false);
 
 				UpdatePowerStatus();
 				Flush();
@@ -336,6 +343,7 @@ void FCrashReportAnalyticsSessionSummary::Initialize(const FString& ProcessGroup
 
 					while (!bShutdown)
 					{
+						CrcAnalyticsProperties::TickCount.Update(PropertyStore.Get(), [](uint32& Actual) { ++Actual; return true; });
 						CrashReportClientUtils::TickPlatformSpecific();
 
 						// Monitor the power level.
@@ -512,6 +520,7 @@ void FCrashReportAnalyticsSessionSummary::OnMonitoredAppDeath(FProcHandle& Handl
 	if (IsValid() && bMonitoredAppDeathRecorded.compare_exchange_strong(Expected, true))
 	{
 		CrcAnalyticsProperties::MonitoredAppDeathTimestamp.Set(PropertyStore.Get(), FDateTime::UtcNow());
+		LogEvent(TEXT("App/Death"));
 
 		int32 ExitCode;
 		if (Handle.IsValid() && FPlatformProcess::GetProcReturnCode(Handle, &ExitCode))
@@ -728,3 +737,9 @@ bool FCrashReportAnalyticsSessionSummary::UpdatePowerStatus()
 
 	return bShouldFlush;
 }
+
+void FCrashReportAnalyticsSessionSummary::OnCheckingForCrash(bool bChecking)
+{
+	CrcAnalyticsProperties::ReadingPipeForCrash.Set(PropertyStore.Get(), bChecking);
+	Flush();
+}
diff --git a/Engine/Source/Programs/CrashReportClient/Private/CrashReportAnalyticsSessionSummary.h b/Engine/Source/Programs/CrashReportClient/Private/CrashReportAnalyticsSessionSummary.h
--- a/Engine/Source/Programs/CrashReportClient/Private/CrashReportAnalyticsSessionSummary.h
+++ b/Engine/Source/Programs/CrashReportClient/Private/CrashReportAnalyticsSessionSummary.h
@@ -58,6 +58,9 @@ public:
 	/** Invoked when the system gently signal CRC to quit. */
 	void OnQuitSignal();
 
+	/** Invoked when CRC is reading the pipe to detect a crash. */
+	void OnCheckingForCrash(bool bChecking);
+
 	/** Invoked when CRC starts to process a crash/ensure/stall.*/
 	void OnCrashReportStarted(ECrashContextType Type, const TCHAR* ErrorMsg);
 
diff --git a/Engine/Source/Programs/CrashReportClient/Private/CrashReportClientApp.cpp b/Engine/Source/Programs/CrashReportClient/Private/CrashReportClientApp.cpp
--- a/Engine/Source/Programs/CrashReportClient/Private/CrashReportClientApp.cpp
+++ b/Engine/Source/Programs/CrashReportClient/Private/CrashReportClientApp.cpp
@@ -5,6 +5,7 @@
 #include "Misc/Parse.h"
 #include "Misc/CommandLine.h"
 #include "Misc/QueuedThreadPool.h"
+#include "Misc/ScopeExit.h"
 #include "Internationalization/Internationalization.h"
 #include "Math/Vector2D.h"
 #include "Misc/ConfigCacheIni.h"
@@ -679,6 +680,14 @@ bool IsCrashReportAvailable(uint32 WatchedProcess, FSharedCrashContext& CrashCon
 {
 	TArray<uint8> Buffer;
 
+	// Temporary code: CRC suspiciously goes away without reason too often. It is either killed or died badly with an exception that skip over
+	//                 SEH exception handling (like STACK_BUFFER_OVERFLOW). One hypothesis is that we try to read a pipe to a corrupted process.
+	//                 So for a limited time, we are going to track when we are reading the pipe. This is IO extensive, because we need to flush
+	//                 the CRC analytics summary file each time, so we probably want to turn it off later once we figured out whether CRC crashes
+	//                 while reading the pipe or not.
+	FCrashReportAnalyticsSessionSummary::Get().OnCheckingForCrash(true);
+	ON_SCOPE_EXIT { FCrashReportAnalyticsSessionSummary::Get().OnCheckingForCrash(false); };
+
 	// Is data available on the pipe.
 	if (FPlatformProcess::ReadPipeToArray(ReadPipe, Buffer))
 	{
