Removed CrashReportClient analytic field 'MonitorQueryingPipe' that was temporary added to verify if CRC crashed while reading the pipe.
  - The data show no evidence that CRC is crashing there. Capturing this state is I/O expensive and not required moving forward.

diff --git a/Engine/Source/Programs/CrashReportClient/Private/CrashReportAnalyticsSessionSummary.cpp b/Engine/Source/Programs/CrashReportClient/Private/CrashReportAnalyticsSessionSummary.cpp
--- a/Engine/Source/Programs/CrashReportClient/Private/CrashReportAnalyticsSessionSummary.cpp
+++ b/Engine/Source/Programs/CrashReportClient/Private/CrashReportAnalyticsSessionSummary.cpp
@@ -67,7 +67,8 @@ namespace CrcAnalyticsProperties
 	//     - V4 -> Stripped ensure callstack from the ensure error message to remove noise in the diagnostic log.
 	//     - V5 -> Measured time to stack-walk, gather files and addded stall count.
 	//     - V6 -> Added MonitorTickCount, MonitorQueryingPipe and App/Death log.
-	constexpr uint32 CrcAnalyticsSummaryVersion = 6;
+	//     - V7 -> Removed MonitorQueryingPipe. It was added to detect if CRC crashed while reading the pipe. Data showed that wasn't the case.
+	constexpr uint32 CrcAnalyticsSummaryVersion = 7;
 
 	/** The exit code of the monitored application. */
 	static const TAnalyticsProperty<int32> MonitoredAppExitCode(TEXT("ExitCode"));
@@ -118,8 +119,6 @@ namespace CrcAnalyticsProperties
 	static const TAnalyticsProperty<uint32> StallCount(TEXT("MonitorStallCount"));
 	/** The worst unattended report time measured (if any). The user is not involved, so it measure how fast CRC can process a crash, especially ensures and stalls. */
 	static const TAnalyticsProperty<float> LonguestUnattendedReportSecs(TEXT("MonitorLongestUnattendedReportSecs"));
-	/** Flags raised when the CRC is reading the pipe. */
-	static const TAnalyticsProperty<bool> ReadingPipeForCrash(TEXT("MonitorQueryingPipe"));
 }
 
 namespace CrashReportClientUtils
@@ -317,7 +316,6 @@ void FCrashReportAnalyticsSessionSummary::Initialize(const FString& ProcessGroup
 				CrcAnalyticsProperties::EnsureCount.Set(PropertyStore.Get(), 0);
 				CrcAnalyticsProperties::AssertCount.Set(PropertyStore.Get(), 0);
 				CrcAnalyticsProperties::StallCount.Set(PropertyStore.Get(), 0);
-				CrcAnalyticsProperties::ReadingPipeForCrash.Set(PropertyStore.Get(), false);
 
 				UpdatePowerStatus();
 				Flush();
@@ -737,9 +735,3 @@ bool FCrashReportAnalyticsSessionSummary::UpdatePowerStatus()
 
 	return bShouldFlush;
 }
-
-void FCrashReportAnalyticsSessionSummary::OnCheckingForCrash(bool bChecking)
-{
-	CrcAnalyticsProperties::ReadingPipeForCrash.Set(PropertyStore.Get(), bChecking);
-	Flush();
-}
diff --git a/Engine/Source/Programs/CrashReportClient/Private/CrashReportAnalyticsSessionSummary.h b/Engine/Source/Programs/CrashReportClient/Private/CrashReportAnalyticsSessionSummary.h
--- a/Engine/Source/Programs/CrashReportClient/Private/CrashReportAnalyticsSessionSummary.h
+++ b/Engine/Source/Programs/CrashReportClient/Private/CrashReportAnalyticsSessionSummary.h
@@ -58,9 +58,6 @@ public:
 	/** Invoked when the system gently signal CRC to quit. */
 	void OnQuitSignal();
 
-	/** Invoked when CRC is reading the pipe to detect a crash. */
-	void OnCheckingForCrash(bool bChecking);
-
 	/** Invoked when CRC starts to process a crash/ensure/stall.*/
 	void OnCrashReportStarted(ECrashContextType Type, const TCHAR* ErrorMsg);
 
diff --git a/Engine/Source/Programs/CrashReportClient/Private/CrashReportClientApp.cpp b/Engine/Source/Programs/CrashReportClient/Private/CrashReportClientApp.cpp
--- a/Engine/Source/Programs/CrashReportClient/Private/CrashReportClientApp.cpp
+++ b/Engine/Source/Programs/CrashReportClient/Private/CrashReportClientApp.cpp
@@ -696,14 +696,6 @@ bool IsCrashReportAvailable(uint32 WatchedProcess, FSharedCrashContext& CrashCon
 {
 	TArray<uint8> Buffer;
 
-	// Temporary code: CRC suspiciously goes away without reason too often. It is either killed or died badly with an exception that skip over
-	//                 SEH exception handling (like STACK_BUFFER_OVERFLOW). One hypothesis is that we try to read a pipe to a corrupted process.
-	//                 So for a limited time, we are going to track when we are reading the pipe. This is IO extensive, because we need to flush
-	//                 the CRC analytics summary file each time, so we probably want to turn it off later once we figured out whether CRC crashes
-	//                 while reading the pipe or not.
-	FCrashReportAnalyticsSessionSummary::Get().OnCheckingForCrash(true);
-	ON_SCOPE_EXIT { FCrashReportAnalyticsSessionSummary::Get().OnCheckingForCrash(false); };
-
 	// Is data available on the pipe.
 	if (FPlatformProcess::ReadPipeToArray(ReadPipe, Buffer))
 	{
