Fixed hang in crash reporter client when reporting bug in unattended mode.
  - Ensured the new FTSTicker get ticked when running CRC in unattended mode.

diff --git a/Engine/Source/Programs/CrashReportClient/Private/MainLoopTiming.cpp b/Engine/Source/Programs/CrashReportClient/Private/MainLoopTiming.cpp
--- a/Engine/Source/Programs/CrashReportClient/Private/MainLoopTiming.cpp
+++ b/Engine/Source/Programs/CrashReportClient/Private/MainLoopTiming.cpp
@@ -26,6 +26,7 @@ void FMainLoopTiming::Tick()
 	// Tick app logic
 	FTaskGraphInterface::Get().ProcessThreadUntilIdle(ENamedThreads::GameThread);
 	FTicker::GetCoreTicker().Tick(ActualDeltaTime);
+	FTSTicker::GetCoreTicker().Tick(ActualDeltaTime);
 
 #if !CRASH_REPORT_UNATTENDED_ONLY
 	// Tick SlateApplication
