- Fix aligned allocation reliated compile fails in targets building against old OSX versions.

diff --git a/Engine/Source/Programs/CrashReportClient/CrashReportClient.Target.cs b/Engine/Source/Programs/CrashReportClient/CrashReportClient.Target.cs
--- a/Engine/Source/Programs/CrashReportClient/CrashReportClient.Target.cs
+++ b/Engine/Source/Programs/CrashReportClient/CrashReportClient.Target.cs
@@ -49,6 +49,9 @@ public class CrashReportClientTarget : TargetRules
 		// Epic Games Launcher needs to run on OS X 10.9, so CrashReportClient needs this as well
 		bEnableOSX109Support = true;
 
+		// Disable automation tests to avoid compilation errors related to alignas, caused by setting bEnableOSX109Support to true above. Old OSX versions do not fully support this until C++17
+		bForceDisableAutomationTests = true;
+
 		// Need to disable the bundled version of dbghelp so that CrashDebugHelper can load dbgeng.dll.
 		WindowsPlatform.bUseBundledDbgHelp = false;
 
diff --git a/Engine/Source/Programs/Mac/UnrealEditorServices/UnrealEditorServices.Target.cs b/Engine/Source/Programs/Mac/UnrealEditorServices/UnrealEditorServices.Target.cs
--- a/Engine/Source/Programs/Mac/UnrealEditorServices/UnrealEditorServices.Target.cs
+++ b/Engine/Source/Programs/Mac/UnrealEditorServices/UnrealEditorServices.Target.cs
@@ -25,5 +25,8 @@ public class UnrealEditorServicesTarget : TargetRules
 
 		// We still need to support old versions of the engine that are compatible with OS X 10.9
 		bEnableOSX109Support = true;
+
+		// Disable automation tests to avoid compilation errors related to alignas, caused by setting bEnableOSX109Support to true above. Old OSX versions do not fully support this until C++17
+		bForceDisableAutomationTests = true;
 	}
 }
