Fix FGenericCrashContext::Initialize() to initialize the GameName to UE-ProjectName instead of UE5-GameName

diff --git a/Engine/Source/Programs/CrashReportClient/Private/CrashReportClientApp.cpp b/Engine/Source/Programs/CrashReportClient/Private/CrashReportClientApp.cpp
--- a/Engine/Source/Programs/CrashReportClient/Private/CrashReportClientApp.cpp
+++ b/Engine/Source/Programs/CrashReportClient/Private/CrashReportClientApp.cpp
@@ -312,7 +312,14 @@ SubmitCrashReportResult RunWithUI(FPlatformErrorReport ErrorReport, bool bImplic
 	TSharedRef<SCrashReportClient> ClientControl = SNew(SCrashReportClient, CrashReportClient, bImplicitSend);
 
 	FString CrashedAppName = FPrimaryCrashProperties::Get()->IsValid() ? FPrimaryCrashProperties::Get()->GameName : TEXT("");
-	CrashedAppName.RemoveFromStart(TEXT("UE4-"));
+	// GameNames have taken on a number of prefixes over the years. Try to strip them all off.
+	if (!CrashedAppName.RemoveFromStart(TEXT("UE4-")))
+	{
+		if (!CrashedAppName.RemoveFromStart(TEXT("UE5-")))
+		{
+			CrashedAppName.RemoveFromStart(TEXT("UE-"));
+		}
+	}
 	CrashedAppName.RemoveFromEnd(TEXT("Game"));
 
 	const FString CrashedAppString = NSLOCTEXT("CrashReportClient", "CrashReporterTitle", "Crash Reporter").ToString();
diff --git a/Engine/Source/Programs/CrashReportClient/Private/SCrashReportClient.cpp b/Engine/Source/Programs/CrashReportClient/Private/SCrashReportClient.cpp
--- a/Engine/Source/Programs/CrashReportClient/Private/SCrashReportClient.cpp
+++ b/Engine/Source/Programs/CrashReportClient/Private/SCrashReportClient.cpp
@@ -330,7 +330,14 @@ void SCrashReportClient::ConstructDetailedDialog(const TSharedRef<FCrashReportCl
 void SCrashReportClient::ConstructSimpleDialog(const TSharedRef<FCrashReportClient>& Client, const FText& CrashDetailedMessage)
 {
 	FString CrashedAppName = FPrimaryCrashProperties::Get()->IsValid() ? FPrimaryCrashProperties::Get()->GameName : TEXT("");
-	CrashedAppName.RemoveFromStart(TEXT("UE4-"));
+	// GameNames have taken on a number of prefixes over the years. Try to strip them all off.
+	if (!CrashedAppName.RemoveFromStart(TEXT("UE4-")))
+	{
+		if (!CrashedAppName.RemoveFromStart(TEXT("UE5-")))
+		{
+			CrashedAppName.RemoveFromStart(TEXT("UE-"));
+		}
+	}
 	CrashedAppName.RemoveFromEnd(TEXT("Game"));
 
 	ChildSlot
diff --git a/Engine/Source/Runtime/Core/Private/GenericPlatform/GenericPlatformCrashContext.cpp b/Engine/Source/Runtime/Core/Private/GenericPlatform/GenericPlatformCrashContext.cpp
--- a/Engine/Source/Runtime/Core/Private/GenericPlatform/GenericPlatformCrashContext.cpp
+++ b/Engine/Source/Runtime/Core/Private/GenericPlatform/GenericPlatformCrashContext.cpp
@@ -111,7 +111,7 @@ void FGenericCrashContext::Initialize()
 	NCached::Session.bIsSourceDistribution = FEngineBuildSettings::IsSourceDistribution();
 	NCached::Session.ProcessId = FPlatformProcess::GetCurrentProcessId();
 
-	NCached::Set(NCached::Session.GameName, *FString::Printf(TEXT("UE5-%s"), FApp::GetProjectName()));
+	NCached::Set(NCached::Session.GameName, *FString::Printf(TEXT("UE-%s"), FApp::GetProjectName()));
 	NCached::Set(NCached::Session.GameSessionID, TEXT("")); // Updated by callback
 	NCached::Set(NCached::Session.GameStateName, TEXT("")); // Updated by callback
 	NCached::Set(NCached::Session.UserActivityHint, TEXT("")); // Updated by callback
