Added a command line option to CrashReportClient to disable the 'Submit and Restart' button in case the application doesn't want to be automaticaly restarted.

diff --git a/Engine/Source/Programs/CrashReportClient/Private/CrashReportClientApp.cpp b/Engine/Source/Programs/CrashReportClient/Private/CrashReportClientApp.cpp
--- a/Engine/Source/Programs/CrashReportClient/Private/CrashReportClientApp.cpp
+++ b/Engine/Source/Programs/CrashReportClient/Private/CrashReportClientApp.cpp
@@ -74,6 +74,9 @@ static FString GameNameFromCmd;
 /** GUID of the crash passed via the command line. */
 static FString CrashGUIDFromCmd;
 
+/** When the application invoking CRC cannot be restarted, force hiding the submit and restart button. */
+static bool bForceHideSubmitAndRestartButtonFromCmd = false;
+
 /** If we are implicitly sending its assumed we are also unattended for now */
 static bool bImplicitSendFromCmd = false;
 /** If we want to enable analytics */
@@ -167,6 +170,11 @@ void ParseCommandLine(const TCHAR* CommandLine)
 			AnalyticsEnabledFromCmd = false;
 		}
 
+		if (Switches.Contains(TEXT("HideSubmitAndRestart")))
+		{
+			bForceHideSubmitAndRestartButtonFromCmd = true;
+		}
+
 		CrashGUIDFromCmd = Params.FindRef(TEXT("CrashGUID"));
 		MonitorPid = FPlatformString::Atoi64(*Params.FindRef(TEXT("MONITOR")));
 		MonitorReadPipe = (void*) FPlatformString::Atoi64(*Params.FindRef(TEXT("READ")));
@@ -309,7 +317,8 @@ SubmitCrashReportResult RunWithUI(FPlatformErrorReport ErrorReport, bool bImplic
 
 	// Open up the app window
 	// bImplicitSend now implies bSimpleDialog i.e. immediately send the report and notify the user without requesting input
-	TSharedRef<SCrashReportClient> ClientControl = SNew(SCrashReportClient, CrashReportClient, bImplicitSend);
+	TSharedRef<SCrashReportClient> ClientControl = SNew(SCrashReportClient, CrashReportClient, bImplicitSend)
+		.bHideSubmitAndRestart(bForceHideSubmitAndRestartButtonFromCmd);
 
 	FString CrashedAppName = FPrimaryCrashProperties::Get()->IsValid() ? FPrimaryCrashProperties::Get()->GameName : TEXT("");
 	// GameNames have taken on a number of prefixes over the years. Try to strip them all off.
diff --git a/Engine/Source/Programs/CrashReportClient/Private/SCrashReportClient.cpp b/Engine/Source/Programs/CrashReportClient/Private/SCrashReportClient.cpp
--- a/Engine/Source/Programs/CrashReportClient/Private/SCrashReportClient.cpp
+++ b/Engine/Source/Programs/CrashReportClient/Private/SCrashReportClient.cpp
@@ -45,6 +45,7 @@
 {
 	CrashReportClient = Client;
 	bHasUserCommentErrors = false;
+	bHideSubmitAndRestart = InArgs._bHideSubmitAndRestart;
 
 	FText CrashDetailedMessage = LOCTEXT("CrashDetailed", "We are very sorry that this crash occurred. Our goal is to prevent crashes like this from occurring in the future. Please help us track down and fix this crash by providing detailed information about what you were doing so that we may reproduce the crash and fix it quickly. You can also log a Bug Report with us using the <a id=\"browser\" href=\"https://epicsupport.force.com/unrealengine/s/\" style=\"Hyperlink\">Bug Submission Form</> and work directly with support staff to report this issue.\n\nThanks for your help in improving the Unreal Engine.");
 	if (FPrimaryCrashProperties::Get()->IsValid())
@@ -318,8 +319,8 @@
 					.Text(LOCTEXT("SendAndRestartEditor", "Send and Restart"))
 					.OnClicked(Client, &FCrashReportClient::SubmitAndRestart)
 					.IsEnabled(this, &SCrashReportClient::IsSendEnabled)
-					.Visibility( FCrashReportCoreConfig::Get().GetHideRestartOption() ? EVisibility::Collapsed : EVisibility::Visible )
-				]			
+					.Visibility( bHideSubmitAndRestart || FCrashReportCoreConfig::Get().GetHideRestartOption() ? EVisibility::Collapsed : EVisibility::Visible )
+				]
 			]
 		]
 	];
diff --git a/Engine/Source/Programs/CrashReportClient/Private/SCrashReportClient.h b/Engine/Source/Programs/CrashReportClient/Private/SCrashReportClient.h
--- a/Engine/Source/Programs/CrashReportClient/Private/SCrashReportClient.h
+++ b/Engine/Source/Programs/CrashReportClient/Private/SCrashReportClient.h
@@ -20,9 +20,13 @@ public:
 	 * Slate arguments
 	 */
 	SLATE_BEGIN_ARGS(SCrashReportClient)
+		: _bHideSubmitAndRestart(false)
 	{
 	}
 
+	/** Should the Submit and Send button be hitten. This can be overriden by a platform settings in the crash report config ini file. */
+	SLATE_ARGUMENT(bool, bHideSubmitAndRestart)
+
 	SLATE_END_ARGS()
 
 	/**
@@ -71,6 +75,7 @@ private:
 	TSharedPtr<SMultiLineEditableTextBox> CrashDetailsInformation;
 
 	bool bHasUserCommentErrors;
+	bool bHideSubmitAndRestart;
 };
 
 #endif // !CRASH_REPORT_UNATTENDED_ONLY
diff --git a/Engine/Source/Runtime/Core/Private/Windows/WindowsPlatformCrashContext.cpp b/Engine/Source/Runtime/Core/Private/Windows/WindowsPlatformCrashContext.cpp
--- a/Engine/Source/Runtime/Core/Private/Windows/WindowsPlatformCrashContext.cpp
+++ b/Engine/Source/Runtime/Core/Private/Windows/WindowsPlatformCrashContext.cpp
@@ -53,6 +53,10 @@
 #define NOINITCRASHREPORTER 0
 #endif
 
+#ifndef DISABLE_CRC_SUBMIT_AND_RESTART_BUTTON
+#define DISABLE_CRC_SUBMIT_AND_RESTART_BUTTON 0
+#endif
+
 #define CR_CLIENT_MAX_PATH_LEN 265
 #define CR_CLIENT_MAX_ARGS_LEN 256
 
@@ -332,6 +336,11 @@ FProcHandle LaunchCrashReportClient(void** OutWritePipe, void** OutReadPipe, uin
 		FCString::Strncat(CrashReporterClientArgs, AnalyticsSummaryGuid, CR_CLIENT_MAX_ARGS_LEN);
 	}
 
+#if DISABLE_CRC_SUBMIT_AND_RESTART_BUTTON
+	// Prevent showing the button if the application cannot be restarted.
+	FCString::Strncat(CrashReporterClientArgs, TEXT(" -HideSubmitAndRestart"), CR_CLIENT_MAX_ARGS_LEN);
+#endif
+
 	// Parse commandline arguments relevant to pass to the client. Note that since we run this from static initialization
 	// FCommandline has not yet been initialized, instead we need to use the OS provided methods.
 	{
@@ -854,6 +863,12 @@ int32 ReportCrashUsingCrashReportClient(FWindowsPlatformCrashContext& InContext,
 					CrashReportClientArguments += TEXT(" -waitforattach");
 				}
 #endif
+
+#if DISABLE_CRC_SUBMIT_AND_RESTART_BUTTON
+				// Prevent showing the button if the application cannot be restarted.
+				CrashReportClientArguments += TEXT(" -HideSubmitAndRestart");
+#endif
+
 				bCrashReporterRan = FPlatformProcess::CreateProc(CrashReporterClientPath, *CrashReportClientArguments, true, false, false, NULL, 0, NULL, NULL).IsValid();
 			}
 
