Add configurable button to copy crash report files to clipboard.

diff --git a/Engine/Source/Programs/CrashReportClient/Private/CrashReportClient.cpp b/Engine/Source/Programs/CrashReportClient/Private/CrashReportClient.cpp
--- a/Engine/Source/Programs/CrashReportClient/Private/CrashReportClient.cpp
+++ b/Engine/Source/Programs/CrashReportClient/Private/CrashReportClient.cpp
@@ -111,6 +111,19 @@ FReply FCrashReportClient::Close()
 	return FReply::Handled();
 }
 
+#if PLATFORM_WINDOWS
+extern void CopyDiagnosticFilesToClipboard(TConstArrayView<FString> Files);
+#endif
+
+#if PLATFORM_WINDOWS
+FReply FCrashReportClient::CopyFilesToClipboard()
+{
+	TArray<FString> Files = FPlatformErrorReport(ErrorReport.GetReportDirectory()).GetFilesToUpload();
+	CopyDiagnosticFilesToClipboard(Files);
+	return FReply::Handled();
+}
+#endif
+
 FReply FCrashReportClient::Submit()
 {
 	bSendData = true;
diff --git a/Engine/Source/Programs/CrashReportClient/Private/CrashReportClient.h b/Engine/Source/Programs/CrashReportClient/Private/CrashReportClient.h
--- a/Engine/Source/Programs/CrashReportClient/Private/CrashReportClient.h
+++ b/Engine/Source/Programs/CrashReportClient/Private/CrashReportClient.h
@@ -76,6 +76,11 @@ public:
 
 	/** Closes the crash report client allowing data to finish being sent. */
 	FReply Close();
+	
+#if PLATFORM_WINDOWS
+	/** Copy report files to the clipboard for user to immediately request help with. */
+	FReply CopyFilesToClipboard();
+#endif
 
 	/**
 	 * Respond to the user pressing Submit
diff --git a/Engine/Source/Programs/CrashReportClient/Private/SCrashReportClient.cpp b/Engine/Source/Programs/CrashReportClient/Private/SCrashReportClient.cpp
--- a/Engine/Source/Programs/CrashReportClient/Private/SCrashReportClient.cpp
+++ b/Engine/Source/Programs/CrashReportClient/Private/SCrashReportClient.cpp
@@ -295,6 +295,21 @@ void SCrashReportClient::ConstructDetailedDialog(const TSharedRef<FCrashReportCl
 					SNew(SSpacer)
 				]
 
+#if PLATFORM_WINDOWS
+				+SHorizontalBox::Slot()
+				.HAlign(HAlign_Center)
+				.VAlign(VAlign_Center)
+				.AutoWidth()
+				.Padding( FMargin(6) )
+				[
+					SNew(SButton)
+					.ContentPadding( FMargin(8,2) )
+					.Text(LOCTEXT("CopyFiles", "Copy Files To Clipboard"))
+					.OnClicked(Client, &FCrashReportClient::CopyFilesToClipboard)
+					.Visibility(FCrashReportCoreConfig::Get().IsAllowedToCopyFilesToClipboard() ? EVisibility::Visible : EVisibility::Hidden)
+				]
+#endif
+
 				+SHorizontalBox::Slot()
 				.HAlign(HAlign_Center)
 				.VAlign(VAlign_Center)
diff --git a/Engine/Source/Programs/CrashReportClient/Private/SCrashReportClient.h b/Engine/Source/Programs/CrashReportClient/Private/SCrashReportClient.h
--- a/Engine/Source/Programs/CrashReportClient/Private/SCrashReportClient.h
+++ b/Engine/Source/Programs/CrashReportClient/Private/SCrashReportClient.h
@@ -69,6 +69,11 @@ private:
 	/** Whether the send buttons are enabled. */
 	bool IsSendEnabled() const;
 
+#if PLATFORM_WINDOWS
+	/** Whether the copy to clipboard button is available. */
+	bool IsCopyToClipboardEnabled() const;
+#endif
+
 	/** Crash report client implementation object */
 	TSharedPtr<FCrashReportClient> CrashReportClient;
 
diff --git a/Engine/Source/Programs/CrashReportClient/Private/Windows/CrashReportClientMainWindows.cpp b/Engine/Source/Programs/CrashReportClient/Private/Windows/CrashReportClientMainWindows.cpp
--- a/Engine/Source/Programs/CrashReportClient/Private/Windows/CrashReportClientMainWindows.cpp
+++ b/Engine/Source/Programs/CrashReportClient/Private/Windows/CrashReportClientMainWindows.cpp
@@ -13,6 +13,48 @@
 #include "CrashReportAnalyticsSessionSummary.h"
 #endif
 
+#include "Windows/AllowWindowsPlatformTypes.h"
+	#include <ShlObj.h>
+#include "Windows/HideWindowsPlatformTypes.h"
+
+void CopyDiagnosticFilesToClipboard(TConstArrayView<FString> Files)
+{
+	if( OpenClipboard(GetActiveWindow()) )
+	{
+		verify(EmptyClipboard());
+		HGLOBAL GlobalMem;
+		SIZE_T RequiredSize = sizeof(DROPFILES) + sizeof(TCHAR);
+		for (const FString& File : Files)
+		{
+			RequiredSize += (File.Len() * sizeof(TCHAR)) + sizeof(TCHAR);
+		}
+		GlobalMem = GlobalAlloc( GMEM_MOVEABLE, RequiredSize );
+		check(GlobalMem);
+		uint8* Data = (uint8*) GlobalLock( GlobalMem );
+		DROPFILES* Drop = (DROPFILES*)Data;
+		Drop->pFiles = sizeof(DROPFILES);
+		Drop->fWide = 1;
+		TCHAR* Dest = (TCHAR*)(Data + sizeof(DROPFILES));
+		TCHAR* End = (TCHAR*)(Data + RequiredSize);
+		for (const FString& File : Files)
+		{
+			FCString::Strncpy(Dest, *File, End - Dest);	
+			Dest += (File.Len() + 1);
+		}
+		GlobalUnlock( GlobalMem );
+		if( SetClipboardData( CF_HDROP, GlobalMem ) == NULL )
+		{
+			UE_LOG(LogWindows, Fatal,TEXT("SetClipboardData failed with error code %i"), (uint32)GetLastError() );
+		}
+
+		verify(CloseClipboard());
+	}
+	else
+	{
+		UE_LOG(LogWindows, Warning, TEXT("OpenClipboard failed with error code %i"), (uint32)GetLastError());
+	}
+}
+
 #if CRASH_REPORT_WITH_MTBF && !PLATFORM_SEH_EXCEPTIONS_DISABLED
 
 static ANSICHAR CrashStackTrace[8*1024] = {0};
diff --git a/Engine/Source/Runtime/CrashReportCore/Private/CrashReportCoreConfig.cpp b/Engine/Source/Runtime/CrashReportCore/Private/CrashReportCoreConfig.cpp
--- a/Engine/Source/Runtime/CrashReportCore/Private/CrashReportCoreConfig.cpp
+++ b/Engine/Source/Runtime/CrashReportCore/Private/CrashReportCoreConfig.cpp
@@ -129,6 +129,14 @@ void FCrashReportCoreConfig::SetProjectConfigOverrides(const FConfigFile& InConf
 		{
 			bIsAllowedToCloseWithoutSending = FCString::ToBool(*IsAllowedToCloseWithoutSendingValue->GetValue());
 		}
+
+#if PLATFORM_WINDOWS
+		const FConfigValue* IsAllowedToCopyFilesToClipboard = Section->Find(TEXT("bIsAllowedToCopyFilesToClipboard"));
+		if (IsAllowedToCopyFilesToClipboard  != nullptr)
+		{
+			bIsAllowedToCopyFilesToClipboard  = FCString::ToBool(*IsAllowedToCopyFilesToClipboard->GetValue());
+		}
+#endif
 	}
 }
 
diff --git a/Engine/Source/Runtime/CrashReportCore/Public/CrashReportCoreConfig.h b/Engine/Source/Runtime/CrashReportCore/Public/CrashReportCoreConfig.h
--- a/Engine/Source/Runtime/CrashReportCore/Public/CrashReportCoreConfig.h
+++ b/Engine/Source/Runtime/CrashReportCore/Public/CrashReportCoreConfig.h
@@ -105,6 +105,13 @@ struct FCrashReportCoreConfig
 		return bIsAllowedToCloseWithoutSending;
 	}
 
+#if PLATFORM_WINDOWS
+	const bool IsAllowedToCopyFilesToClipboard() const
+	{
+		return bIsAllowedToCopyFilesToClipboard;
+	}
+#endif
+
 	int GetUserCommentSizeLimit() const
 	{
 		return UserCommentSizeLimit;
@@ -165,6 +172,11 @@ protected:
 	/** Whether the user is allowed to close the crash reporter without sending a report */
 	bool bIsAllowedToCloseWithoutSending;
 
+#if PLATFORM_WINDOWS
+	/** Whether the user is allowed to copy the report files to the clipboard */
+	bool bIsAllowedToCopyFilesToClipboard;
+#endif
+
 	/** Size limit for the description of multi-line text */
 	int UserCommentSizeLimit;
 };
