Add an API for adding many portable thread callstacks at once to avoid re-reading process module list.
Cuts time to record a crash context by as much as 30s in the presence of some security software.

diff --git a/Engine/Source/Programs/CrashReportClient/Private/CrashReportClientApp.cpp b/Engine/Source/Programs/CrashReportClient/Private/CrashReportClientApp.cpp
--- a/Engine/Source/Programs/CrashReportClient/Private/CrashReportClientApp.cpp
+++ b/Engine/Source/Programs/CrashReportClient/Private/CrashReportClientApp.cpp
@@ -507,6 +507,8 @@ FPlatformErrorReport CollectErrorReport(FRecoveryService* RecoveryService, uint3
 	// Initialize the stack walking for the monitored process (effectively overriding this process stack walking functionality)
 	FPlatformStackWalk::InitStackWalkingForProcess(ProcessHandle);
 
+	TArray<TArray<uint64>> ThreadCallStacks;
+	ThreadCallStacks.Reserve(SharedCrashContext.NumThreads);
 	for (uint32 ThreadIdx = 0; ThreadIdx < SharedCrashContext.NumThreads; ThreadIdx++)
 	{
 		const uint32 ThreadId = SharedCrashContext.ThreadIds[ThreadIdx];
@@ -546,13 +548,15 @@ FPlatformErrorReport CollectErrorReport(FRecoveryService* RecoveryService, uint3
 			CR_MAX_STACK_FRAMES,
 			PlatformContext.Get()
 		);
+		
+		ThreadCallStacks.Emplace(TArray<uint64>(StackFrames, StackFrameCount));		
 
-		CrashContext.AddPortableThreadCallStack(
-			SharedCrashContext.ThreadIds[ThreadIdx],
-			&SharedCrashContext.ThreadNames[ThreadIdx*CR_MAX_THREAD_NAME_CHARS],
-			StackFrames,
-			StackFrameCount
-		);
+		// CrashContext.AddPortableThreadCallStack(
+		// 	SharedCrashContext.ThreadIds[ThreadIdx],
+		// 	&SharedCrashContext.ThreadNames[ThreadIdx*CR_MAX_THREAD_NAME_CHARS],
+		// 	StackFrames,
+		// 	StackFrameCount
+		// );
 
 		// Add the crashing stack specifically. Is this really needed?
 		if (ThreadId == SharedCrashContext.CrashingThreadId)
@@ -595,6 +599,21 @@ FPlatformErrorReport CollectErrorReport(FRecoveryService* RecoveryService, uint3
 		}
 	}
 
+	{
+		TArray<FThreadCallStack> Threads;
+		Threads.Reserve(SharedCrashContext.NumThreads);
+		for (uint32 ThreadIdx = 0; ThreadIdx < SharedCrashContext.NumThreads; ++ThreadIdx)
+		{
+			Threads.Add({
+				MakeArrayView(ThreadCallStacks[ThreadIdx]),
+				&SharedCrashContext.ThreadNames[ThreadIdx*CR_MAX_THREAD_NAME_CHARS],
+				SharedCrashContext.ThreadIds[ThreadIdx],
+			});
+		}
+		CrashContext.AddPortableThreadCallStacks(Threads);
+	}
+
+
 	FCrashReportAnalyticsSessionSummary::Get().OnCrashReportGatheringFiles();
 
 	// If the path is not set it is most likely that we have crashed during static init, in which case we need to construct a directory ourself.
diff --git a/Engine/Source/Runtime/Core/Private/GenericPlatform/GenericPlatformCrashContext.cpp b/Engine/Source/Runtime/Core/Private/GenericPlatform/GenericPlatformCrashContext.cpp
--- a/Engine/Source/Runtime/Core/Private/GenericPlatform/GenericPlatformCrashContext.cpp
+++ b/Engine/Source/Runtime/Core/Private/GenericPlatform/GenericPlatformCrashContext.cpp
@@ -1382,6 +1382,14 @@ void FGenericCrashContext::GetPortableCallStack(const uint64* StackFrames, int32
 	}
 }
 
+void FGenericCrashContext::AddPortableThreadCallStacks(TConstArrayView<FThreadCallStack> Threads)
+{
+	for (const FThreadCallStack& Thread : Threads)
+	{
+		AddPortableThreadCallStack(Thread.ThreadId, Thread.ThreadName, Thread.StackFrames.GetData(), Thread.StackFrames.Num());
+	}
+}
+
 void FGenericCrashContext::AddPortableThreadCallStack(uint32 ThreadId, const TCHAR* ThreadName, const uint64* StackFrames, int32 NumStackFrames)
 {
 	// Not implemented for generic class
diff --git a/Engine/Source/Runtime/Core/Private/Microsoft/MicrosoftPlatformCrashContext.cpp b/Engine/Source/Runtime/Core/Private/Microsoft/MicrosoftPlatformCrashContext.cpp
--- a/Engine/Source/Runtime/Core/Private/Microsoft/MicrosoftPlatformCrashContext.cpp
+++ b/Engine/Source/Runtime/Core/Private/Microsoft/MicrosoftPlatformCrashContext.cpp
@@ -20,16 +20,25 @@ void FMicrosoftPlatformCrashContext::CaptureAllThreadContexts()
 #endif
 }
 
-void FMicrosoftPlatformCrashContext::AddPortableThreadCallStack(uint32 ThreadId, const TCHAR* ThreadName, const uint64* StackFrames, int32 NumStackFrames)
+void FMicrosoftPlatformCrashContext::AddPortableThreadCallStacks(TConstArrayView<FThreadCallStack> Threads)
 {
 	FModuleHandleArray ProcModuleHandles;
 	GetProcModuleHandles(ProcessHandle, ProcModuleHandles);
 
-	FThreadStackFrames Thread;
-	Thread.ThreadId = ThreadId;
-	Thread.ThreadName = FString(ThreadName);
-	ConvertProgramCountersToStackFrames(ProcessHandle, ProcModuleHandles, StackFrames, NumStackFrames, Thread.StackFrames);
-	ThreadCallStacks.Push(Thread);
+	ThreadCallStacks.Reserve(ThreadCallStacks.Num() + Threads.Num());
+	for (const FThreadCallStack& InThread : Threads)
+	{
+		FThreadStackFrames Thread;
+		Thread.ThreadId = InThread.ThreadId;
+		Thread.ThreadName = FString(InThread.ThreadName);
+		ConvertProgramCountersToStackFrames(ProcessHandle, ProcModuleHandles, InThread.StackFrames.GetData(), InThread.StackFrames.Num(), Thread.StackFrames);
+		ThreadCallStacks.Push(Thread);
+	}
+}
+
+void FMicrosoftPlatformCrashContext::AddPortableThreadCallStack(uint32 ThreadId, const TCHAR* ThreadName, const uint64* StackFrames, int32 NumStackFrames)
+{
+	AddPortableThreadCallStacks({{ MakeArrayView(StackFrames, NumStackFrames), ThreadName, ThreadId }});
 }
 
 void FMicrosoftPlatformCrashContext::SetPortableCallStack(const uint64* StackTrace, int32 StackTraceDepth)
diff --git a/Engine/Source/Runtime/Core/Private/Windows/WindowsPlatformCrashContext.cpp b/Engine/Source/Runtime/Core/Private/Windows/WindowsPlatformCrashContext.cpp
--- a/Engine/Source/Runtime/Core/Private/Windows/WindowsPlatformCrashContext.cpp
+++ b/Engine/Source/Runtime/Core/Private/Windows/WindowsPlatformCrashContext.cpp
@@ -139,7 +139,7 @@ namespace {
 		bool bShouldBeFullCrashDump = InContext.IsFullCrashDump();
 		if (bShouldBeFullCrashDump)
 		{
-			MinidumpType = (MINIDUMP_TYPE)(MiniDumpWithFullMemory | MiniDumpWithFullMemoryInfo | MiniDumpWithHandleData | MiniDumpWithThreadInfo | MiniDumpWithUnloadedModules);
+			MinidumpType = (MINIDUMP_TYPE)(MiniDumpWithFullMemory | MiniDumpWithFullMemoryInfo | MiniDumpWithHandleData | MiniDumpWithThreadInfo );
 		}
 
 		const BOOL Result = MiniDumpWriteDump(Process, GetProcessId(Process), FileHandle, MinidumpType, ExceptionInfo ? &DumpExceptionInfo : NULL, &CrashContextStreamInformation, NULL);
diff --git a/Engine/Source/Runtime/Core/Public/GenericPlatform/GenericPlatformCrashContext.h b/Engine/Source/Runtime/Core/Public/GenericPlatform/GenericPlatformCrashContext.h
--- a/Engine/Source/Runtime/Core/Public/GenericPlatform/GenericPlatformCrashContext.h
+++ b/Engine/Source/Runtime/Core/Public/GenericPlatform/GenericPlatformCrashContext.h
@@ -283,6 +283,13 @@ DECLARE_MULTICAST_DELEGATE_TwoParams(FEngineDataSetDelegate, const FString&, con
 DECLARE_MULTICAST_DELEGATE(FGameDataResetDelegate);
 DECLARE_MULTICAST_DELEGATE_TwoParams(FGameDataSetDelegate, const FString&, const FString&);
 
+struct FThreadCallStack
+{
+	TConstArrayView<uint64> StackFrames;
+	const TCHAR* ThreadName;
+	uint32 ThreadId;
+};
+
 /**
  *	Contains a runtime crash's properties that are common for all platforms.
  *	This may change in the future.
@@ -564,6 +571,7 @@ public:
 	CORE_API virtual void GetModules(TArray<FStackWalkModuleInfo>& OutModules) const;
 	
 	/** Adds a portable callstack for a thread */
+	CORE_API virtual void AddPortableThreadCallStacks(TConstArrayView<FThreadCallStack> Threads);
 	CORE_API virtual void AddPortableThreadCallStack(uint32 ThreadId, const TCHAR* ThreadName, const uint64* StackFrames, int32 NumStackFrames);
 
 	/** Allows platform implementations to copy files to report directory. */
diff --git a/Engine/Source/Runtime/Core/Public/Microsoft/MicrosoftPlatformCrashContext.h b/Engine/Source/Runtime/Core/Public/Microsoft/MicrosoftPlatformCrashContext.h
--- a/Engine/Source/Runtime/Core/Public/Microsoft/MicrosoftPlatformCrashContext.h
+++ b/Engine/Source/Runtime/Core/Public/Microsoft/MicrosoftPlatformCrashContext.h
@@ -15,6 +15,7 @@ struct FMicrosoftPlatformCrashContext : public FGenericCrashContext
 	CORE_API void CaptureAllThreadContexts();
 
 	CORE_API virtual void SetPortableCallStack(const uint64* StackTrace, int32 StackTraceDepth) override;
+	CORE_API virtual void AddPortableThreadCallStacks(TConstArrayView<FThreadCallStack> Threads) override;
 	CORE_API virtual void AddPortableThreadCallStack(uint32 ThreadId, const TCHAR* ThreadName, const uint64* StackFrames, int32 NumStackFrames) override;
 
 protected:
