Safer NSApp terminate: call.

#11271

diff --git a/Engine/Source/Programs/ChaosVisualDebugger/Private/Mac/ChaosVisualDebuggerMac.cpp b/Engine/Source/Programs/ChaosVisualDebugger/Private/Mac/ChaosVisualDebuggerMac.cpp
--- a/Engine/Source/Programs/ChaosVisualDebugger/Private/Mac/ChaosVisualDebuggerMac.cpp
+++ b/Engine/Source/Programs/ChaosVisualDebugger/Private/Mac/ChaosVisualDebuggerMac.cpp
@@ -58,8 +58,10 @@ static FString GSavedCommandLine;
 		}
 	
 	FEngineLoop::AppExit();
-
-	[NSApp terminate: self];
+    
+    dispatch_async(dispatch_get_main_queue(), ^{
+        [NSApp terminate: nil];
+    });
 }
 
 - (NSApplicationTerminateReply)applicationShouldTerminate:(NSApplication *)Sender;
diff --git a/Engine/Source/Programs/CrashReportClient/Private/Mac/CrashReportClientMainMac.cpp b/Engine/Source/Programs/CrashReportClient/Private/Mac/CrashReportClientMainMac.cpp
--- a/Engine/Source/Programs/CrashReportClient/Private/Mac/CrashReportClientMainMac.cpp
+++ b/Engine/Source/Programs/CrashReportClient/Private/Mac/CrashReportClientMainMac.cpp
@@ -49,8 +49,10 @@ static bool GIsUnattended = false;
 	}
 	
 	RunCrashReportClient(*GSavedCommandLine);
-	
-	[NSApp terminate: self];
+    
+    dispatch_async(dispatch_get_main_queue(), ^{
+        [NSApp terminate: nil];
+    });
 }
 
 - (NSApplicationTerminateReply)applicationShouldTerminate:(NSApplication *)Sender;
diff --git a/Engine/Source/Runtime/ApplicationCore/Private/Mac/MacProgramDelegate.cpp b/Engine/Source/Runtime/ApplicationCore/Private/Mac/MacProgramDelegate.cpp
--- a/Engine/Source/Runtime/ApplicationCore/Private/Mac/MacProgramDelegate.cpp
+++ b/Engine/Source/Runtime/ApplicationCore/Private/Mac/MacProgramDelegate.cpp
@@ -55,8 +55,10 @@
 		}
 
 	ProgramExit();
-	
-	[NSApp terminate: self];
+    
+    dispatch_async(dispatch_get_main_queue(), ^{
+        [NSApp terminate: nil];
+    });
 }
 
 - (void) applicationWillTerminate:(NSNotification*)notification
diff --git a/Engine/Source/Runtime/Launch/Private/Mac/LaunchMac.cpp b/Engine/Source/Runtime/Launch/Private/Mac/LaunchMac.cpp
--- a/Engine/Source/Runtime/Launch/Private/Mac/LaunchMac.cpp
+++ b/Engine/Source/Runtime/Launch/Private/Mac/LaunchMac.cpp
@@ -259,7 +259,9 @@ static int32 MacOSVersionCompare(const NSOperatingSystemVersion& VersionA, const
 
 	if (GGuardedMainErrorLevel == 0)
 	{
-		[NSApp terminate: nil];
+        dispatch_async(dispatch_get_main_queue(), ^{
+            [NSApp terminate: nil];
+        });
 	}
 	else
 	{
