Fix unattended debug message

diff --git a/Engine/Source/Programs/CrashReportClient/CrashReportClientEditor.Target.cs b/Engine/Source/Programs/CrashReportClient/CrashReportClientEditor.Target.cs
--- a/Engine/Source/Programs/CrashReportClient/CrashReportClientEditor.Target.cs
+++ b/Engine/Source/Programs/CrashReportClient/CrashReportClientEditor.Target.cs
@@ -36,7 +36,6 @@ public sealed class CrashReportClientEditorTarget : CrashReportClientTarget
 
 		if (bHostRecoverySvc)
 		{
-			throw new BuildException("No longer supported.");
 			AdditionalPlugins.Add("UdpMessaging");
 			AdditionalPlugins.Add("ConcertSyncServer");
 			bCompileWithPluginSupport = true; // Enable Developer plugins (like Concert!)
