Add optional display of company name and endpoint in CRC dialog.

* Disable send buttons if both crash receiver and data router configuration values are unset.
* Optionally display end point in crash report client.
* "Agree to be contacted" text is now configurable with to be able to set company name.

diff --git a/Engine/Source/Programs/CrashReportClient/Private/SCrashReportClient.cpp b/Engine/Source/Programs/CrashReportClient/Private/SCrashReportClient.cpp
--- a/Engine/Source/Programs/CrashReportClient/Private/SCrashReportClient.cpp
+++ b/Engine/Source/Programs/CrashReportClient/Private/SCrashReportClient.cpp
@@ -262,7 +262,7 @@ void SCrashReportClient::ConstructDetailedDialog(const TSharedRef<FCrashReportCl
 					SNew(STextBlock)
 					.AutoWrapText(true)
 					.IsEnabled( !FEngineBuildSettings::IsInternalBuild() )
-					.Text(LOCTEXT("IAgree", "I agree to be contacted by Epic Games via email if additional information about this crash would help fix it."))
+					.Text_Static(&SCrashReportClient::GetContactText)
 				]
 			]
 
@@ -321,6 +321,7 @@ void SCrashReportClient::ConstructDetailedDialog(const TSharedRef<FCrashReportCl
 					.Text(LOCTEXT("Send", "Send and Close"))
 					.OnClicked(Client, &FCrashReportClient::Submit)
 					.IsEnabled(this, &SCrashReportClient::IsSendEnabled)
+					.ToolTipText_Static(&SCrashReportClient::GetSendTooltip)
 				]
 
 				+SHorizontalBox::Slot()
@@ -335,6 +336,7 @@ void SCrashReportClient::ConstructDetailedDialog(const TSharedRef<FCrashReportCl
 					.OnClicked(Client, &FCrashReportClient::SubmitAndRestart)
 					.IsEnabled(this, &SCrashReportClient::IsSendEnabled)
 					.Visibility( bHideSubmitAndRestart || FCrashReportCoreConfig::Get().GetHideRestartOption() ? EVisibility::Collapsed : EVisibility::Visible )
+					.ToolTipText_Static(&SCrashReportClient::GetSendTooltip)
 				]
 			]
 		]
@@ -443,9 +445,58 @@ EVisibility SCrashReportClient::IsHintTextVisible() const
 
 bool SCrashReportClient::IsSendEnabled() const
 {
-	bool bValidAppName = FPrimaryCrashProperties::Get()->IsValid() && !FPrimaryCrashProperties::Get()->GameName.IsEmpty();
+	const bool bValidAppName = FPrimaryCrashProperties::Get()->IsValid() && !FPrimaryCrashProperties::Get()->GameName.IsEmpty();
+	const bool bValidEndPoint = !FCrashReportCoreConfig::Get().GetReceiverAddress().IsEmpty() || !FCrashReportCoreConfig::Get().GetDataRouterURL().IsEmpty();
 
-	return bValidAppName && !bHasUserCommentErrors;
+	return bValidAppName && bValidEndPoint && !bHasUserCommentErrors;
+}
+
+FText SCrashReportClient::GetSendTooltip()
+{
+	// Optionally show a tooltip with the endpoint domain to the user. If the old receiver address is used it is
+	// just a IP number, so there is no point is showing that.
+	static FText CachedDomain = []() {
+		FStringView Endpoint = FCrashReportCoreConfig::Get().GetReceiverAddress();
+		bool bIsUrl = false;
+		if (Endpoint.IsEmpty())
+		{
+			Endpoint = FCrashReportCoreConfig::Get().GetDataRouterURL();
+			bIsUrl = true;
+		}
+		if (Endpoint.IsEmpty())
+		{
+			return LOCTEXT("SendTooltipEmpty", "No server specified.");
+		}
+		if (bIsUrl && FCrashReportCoreConfig::Get().GetShowEndpointInTooltip())
+		{
+			// Show only domain, not full url
+			const int32 Start = Endpoint.StartsWith(TEXT("https://")) ? 8 : 0;
+			Endpoint.RightChopInline(Start);
+			int32 End(INDEX_NONE);
+			Endpoint.FindChar('/', End);
+			Endpoint.LeftInline(End);
+			return FText::Format(LOCTEXT("SendTooltipUrl", "Send to {0}"), FText::FromStringView(Endpoint));
+		}
+		return LOCTEXT("SendTooltip", "Send to server");
+	}();
+	return CachedDomain;
+}
+
+FText SCrashReportClient::GetContactText()
+{
+	static FText CachedContactText = []()
+	{
+		const FStringView Company = FCrashReportCoreConfig::Get().GetCompanyName();
+		if (Company.IsEmpty())
+		{
+			return LOCTEXT("IAgreeNoCompany", "I agree to be contacted via email if additional information about this crash would help fix it.");
+		}
+		return FText::Format(
+			 LOCTEXT("IAgreeCompany", "I agree to be contacted by {0} via email if additional information about this crash would help fix it."),
+			 FText::FromStringView(Company)
+		);
+	}();
+	return CachedContactText;
 }
 
 #undef LOCTEXT_NAMESPACE
diff --git a/Engine/Source/Programs/CrashReportClient/Private/SCrashReportClient.h b/Engine/Source/Programs/CrashReportClient/Private/SCrashReportClient.h
--- a/Engine/Source/Programs/CrashReportClient/Private/SCrashReportClient.h
+++ b/Engine/Source/Programs/CrashReportClient/Private/SCrashReportClient.h
@@ -40,7 +40,6 @@ public:
 	bool IsFinished() { return CrashReportClient->IsUploadComplete() && CrashReportClient->ShouldWindowBeHidden(); }
 
 private:
-
 	/**
 	 * Construct the detailed Slate ui with controls for sending the report and comments
 	 * @param Client Crash report client implementation object
@@ -69,6 +68,12 @@ private:
 	/** Whether the send buttons are enabled. */
 	bool IsSendEnabled() const;
 
+	/** Returns the tooltop text for send button */
+	static FText GetSendTooltip();
+
+	/** Returns the 'allow contact' text */
+	static FText GetContactText();
+
 #if PLATFORM_WINDOWS
 	/** Whether the copy to clipboard button is available. */
 	bool IsCopyToClipboardEnabled() const;
