Add missing fields in CrashReportClient. CrashReporter would fail to sort those crashes properly and it would make these very hard to find.

#tests CrashReporter in shipping with a deployed build as well as with a local build

diff --git a/Engine/Source/Programs/CrashReportClient/Private/CrashReportClientApp.cpp b/Engine/Source/Programs/CrashReportClient/Private/CrashReportClientApp.cpp
--- a/Engine/Source/Programs/CrashReportClient/Private/CrashReportClientApp.cpp
+++ b/Engine/Source/Programs/CrashReportClient/Private/CrashReportClientApp.cpp
@@ -874,6 +874,19 @@ bool FindAndCopyValue(const TMap<FString, FString>& Map, const FString& Key, TCH
 	return false;
 }
 
+template <size_t Size>
+bool FindAndCopyValue(const TMap<FString, FString>& Map, const FString& Key, char(&OutValue)[Size])
+{
+	const FString* ValueString = Map.Find(Key);
+	if (ValueString != nullptr)
+	{
+		FCStringAnsi::Strncpy(OutValue, StringCast<ANSICHAR>(**ValueString).Get(), Size);
+		return true;
+	}
+
+	return false;
+}
+
 static bool LoadTempCrashContextFromFile(FSharedCrashContext& CrashContext, uint64 ProcessID)
 {
 	const FString TempContextFilePath = FGenericCrashContext::GetTempSessionContextFilePath(ProcessID);
@@ -896,18 +909,21 @@ static bool LoadTempCrashContextFromFile(FSharedCrashContext& CrashContext, uint
 	FindAndParseValue(ContextProperties, TEXT("IsInternalBuild"), SessionContext.bIsInternalBuild);
 	FindAndParseValue(ContextProperties, TEXT("IsPerforceBuild"), SessionContext.bIsPerforceBuild);
 	FindAndParseValue(ContextProperties, TEXT("IsSourceDistribution"), SessionContext.bIsSourceDistribution);
-	FindAndCopyValue(ContextProperties, TEXT("EngineVersion"), SessionContext.EngineVersion);
-	FindAndCopyValue(ContextProperties, TEXT("EngineCompatibleVersion"), SessionContext.EngineCompatibleVersion);
-	FindAndCopyValue(ContextProperties, TEXT("BuildVersion"), SessionContext.BuildVersion);
 	FindAndCopyValue(ContextProperties, TEXT("GameName"), SessionContext.GameName);
 	FindAndCopyValue(ContextProperties, TEXT("ExecutableName"), SessionContext.ExecutableName);
+	FindAndCopyValue(ContextProperties, TEXT("BuildConfiguration"), SessionContext.BuildConfigurationName);
 	FindAndCopyValue(ContextProperties, TEXT("GameSessionID"), SessionContext.GameSessionID);
+	FindAndCopyValue(ContextProperties, TEXT("PlatformName"), SessionContext.PlatformName);
+	FindAndCopyValue(ContextProperties, TEXT("PlatformNameIni"), SessionContext.PlatformNameIni);
 	FindAndCopyValue(ContextProperties, TEXT("EngineMode"), SessionContext.EngineMode);
 	FindAndCopyValue(ContextProperties, TEXT("EngineModeEx"), SessionContext.EngineModeEx);
 	FindAndCopyValue(ContextProperties, TEXT("DeploymentName"), SessionContext.DeploymentName);
+	FindAndCopyValue(ContextProperties, TEXT("EngineVersion"), SessionContext.EngineVersion);
+	FindAndCopyValue(ContextProperties, TEXT("EngineCompatibleVersion"), SessionContext.EngineCompatibleVersion);
 	FindAndCopyValue(ContextProperties, TEXT("CommandLine"), SessionContext.CommandLine);
 	FindAndParseValue(ContextProperties, TEXT("LanguageLCID"), SessionContext.LanguageLCID);
 	FindAndCopyValue(ContextProperties, TEXT("AppDefaultLocale"), SessionContext.DefaultLocale);
+	FindAndCopyValue(ContextProperties, TEXT("BuildVersion"), SessionContext.BuildVersion);
 	FindAndParseValue(ContextProperties, TEXT("IsUERelease"), SessionContext.bIsUERelease);
 	FindAndCopyValue(ContextProperties, TEXT("UserName"), SessionContext.UserName);
 	FindAndCopyValue(ContextProperties, TEXT("EpicAccountId"), SessionContext.EpicAccountId);
