Several changes to crash reporting and crash telemetry.

* Renamed some configuration settings in CrashReporterClient(Editor) targets and added comments to clarify their use.
* Merged the different telemetry keys to one.
* Allow telemetry to be set using environment variables.

diff --git a/Engine/Source/Programs/CrashReportClient/CrashReportClient.Target.cs b/Engine/Source/Programs/CrashReportClient/CrashReportClient.Target.cs
--- a/Engine/Source/Programs/CrashReportClient/CrashReportClient.Target.cs
+++ b/Engine/Source/Programs/CrashReportClient/CrashReportClient.Target.cs
@@ -7,21 +7,22 @@ using System.Collections.Generic;
 [SupportedConfigurations(UnrealTargetConfiguration.Debug, UnrealTargetConfiguration.Development, UnrealTargetConfiguration.Shipping)]
 public class CrashReportClientTarget : TargetRules
 {
-	[ConfigFile(ConfigHierarchyType.Engine, "CrashReportClientBuildSettings", "DataRouterFallback")]
-	public string DataRouterFallback;
+	// Default crash reporting url. This is the default value. Overrideable per project.
+	[ConfigFile(ConfigHierarchyType.Engine, "CrashReportClientBuildSettings", "DefaultUrl")]
+	public string DefaultUrl;
 		
-	[ConfigFile(ConfigHierarchyType.Engine, "CrashReportClientBuildSettings", "CompanyName")]
-	public string CompanyName;
+	// Default company name displayed in the UX. Overrideable per project
+	[ConfigFile(ConfigHierarchyType.Engine, "CrashReportClientBuildSettings", "DefaultCompanyName")]
+	public string DefaultCompanyName;
 	
+	// Url for telemetry. (Not used by licensees)
 	[ConfigFile(ConfigHierarchyType.Engine, "CrashReportClientBuildSettings", "TelemetryUrl")]
 	public string TelemetryUrl;
 
-	[ConfigFile(ConfigHierarchyType.Engine, "CrashReportClientBuildSettings", "TelemetryKey_Dev")]
-	public string TelemetryKey_Dev;
+	// Key used by telemetry (Not used by licensees)
+	[ConfigFile(ConfigHierarchyType.Engine, "CrashReportClientBuildSettings", "TelemetryKey")]
+	public string TelemetryKey;
 
-	[ConfigFile(ConfigHierarchyType.Engine, "CrashReportClientBuildSettings", "TelemetryKey_Release")]
-	public string TelemetryKey_Release;
-	
 	public CrashReportClientTarget(TargetInfo Target) : this(Target, true)
 	{}
 
@@ -70,26 +71,25 @@ public class CrashReportClientTarget : TargetRules
 		if (bSetConfiguredDefinitions)
 		{
 			GlobalDefinitions.AddRange(SetupConfiguredDefines(
-				 DataRouterFallback, CompanyName, TelemetryUrl, TelemetryKey_Dev, TelemetryKey_Release));
+				 DefaultUrl, DefaultCompanyName, TelemetryUrl, TelemetryKey));
 		}
 	}
 
 	protected static List<string> SetupConfiguredDefines(
-		string DataRouterFallback, 
-		string CompanyName, 
+		string DefaultUrl, 
+		string DefaultCompanyName, 
 		string TelemetryUrl, 
-		string TelemetryKeyDev,
-		string TelemetryKeyRelease)
+		string TelemetryKey)
 	{
 		var Definitions = new List<string>();
-		if (!string.IsNullOrEmpty(DataRouterFallback))
+		if (!string.IsNullOrEmpty(DefaultUrl))
 		{
-			Definitions.Add($"CRC_DATAROUTER_FALLBACK=\"{DataRouterFallback}\"");
+			Definitions.Add($"CRC_DEFAULT_URL=\"{DefaultUrl}\"");
 		}
 
-		if (!string.IsNullOrEmpty(CompanyName))
+		if (!string.IsNullOrEmpty(DefaultCompanyName))
 		{
-			Definitions.Add($"CRC_COMPANY_NAME_FALLBACK=\"{CompanyName}\"");
+			Definitions.Add($"CRC_DEFAULT_COMPANY_NAME=\"{DefaultCompanyName}\"");
 		}
 
 		if(!string.IsNullOrWhiteSpace(TelemetryUrl))
@@ -97,14 +97,9 @@ public class CrashReportClientTarget : TargetRules
 			Definitions.Add($"CRC_TELEMETRY_URL=\"{TelemetryUrl}\"");
 		}
 
-		if (!string.IsNullOrWhiteSpace(TelemetryKeyDev))
-		{
-			Definitions.Add($"CRC_TELEMETRY_KEY_DEV=\"{TelemetryKeyDev}\"");
-		}
-
-		if (!string.IsNullOrWhiteSpace(TelemetryKeyRelease))
+		if (!string.IsNullOrWhiteSpace(TelemetryKey))
 		{
-			Definitions.Add($"CRC_TELEMETRY_KEY_RELEASE=\"{TelemetryKeyRelease}\"");
+			Definitions.Add($"CRC_TELEMETRY_KEY=\"{TelemetryKey}\"");
 		}
 
 		return Definitions;
diff --git a/Engine/Source/Programs/CrashReportClient/CrashReportClientEditor.Target.cs b/Engine/Source/Programs/CrashReportClient/CrashReportClientEditor.Target.cs
--- a/Engine/Source/Programs/CrashReportClient/CrashReportClientEditor.Target.cs
+++ b/Engine/Source/Programs/CrashReportClient/CrashReportClientEditor.Target.cs
@@ -7,22 +7,20 @@ using UnrealBuildTool;
 public sealed class CrashReportClientEditorTarget : CrashReportClientTarget
 {
 	// Override the configuration values from CrashReportClient with these using another
-	// configuration block: [CrashReportClientEditorBuildSettings]
+	// configuration block: [CrashReportClientEditorBuildSettings]. See CrashReportClient.target.cs for
+	// descriptions of the settings.
 	
-	[ConfigFile(ConfigHierarchyType.Engine, "CrashReportClientEditorBuildSettings", "DataRouterFallback")]
-	public new string DataRouterFallback;
+	[ConfigFile(ConfigHierarchyType.Engine, "CrashReportClientEditorBuildSettings", "DefaultUrl")]
+	public new string DefaultUrl;
 		
-	[ConfigFile(ConfigHierarchyType.Engine, "CrashReportClientEditorBuildSettings", "CompanyName")]
-	public new string CompanyName;
+	[ConfigFile(ConfigHierarchyType.Engine, "CrashReportClientEditorBuildSettings", "DefaultCompanyName")]
+	public new string DefaultCompanyName;
 	
 	[ConfigFile(ConfigHierarchyType.Engine, "CrashReportClientEditorBuildSettings", "TelemetryUrl")]
 	public new string TelemetryUrl;
 
-	[ConfigFile(ConfigHierarchyType.Engine, "CrashReportClientEditorBuildSettings", "TelemetryKey_Dev")]
-	public new string TelemetryKey_Dev;
-
-	[ConfigFile(ConfigHierarchyType.Engine, "CrashReportClientEditorBuildSettings", "TelemetryKey_Release")]
-	public new string TelemetryKey_Release;
+	[ConfigFile(ConfigHierarchyType.Engine, "CrashReportClientEditorBuildSettings", "TelemetryKey")]
+	public new string TelemetryKey;
 	
 	public CrashReportClientEditorTarget(TargetInfo Target) : base(Target, false /* bSetConfiguredDefinitions */)
 	{
@@ -50,7 +48,7 @@ public sealed class CrashReportClientEditorTarget : CrashReportClientTarget
 		
 		// We can now set the configured definitions from CrashReportClientEditorBuildSettings section
 		GlobalDefinitions.AddRange(SetupConfiguredDefines(
-			DataRouterFallback, CompanyName, TelemetryUrl, TelemetryKey_Dev, TelemetryKey_Release));
+			DefaultUrl, DefaultCompanyName, TelemetryUrl, TelemetryKey));
 	}
 	
 }
diff --git a/Engine/Source/Runtime/CrashReportCore/Private/CrashReportAnalytics.cpp b/Engine/Source/Runtime/CrashReportCore/Private/CrashReportAnalytics.cpp
--- a/Engine/Source/Runtime/CrashReportCore/Private/CrashReportAnalytics.cpp
+++ b/Engine/Source/Runtime/CrashReportCore/Private/CrashReportAnalytics.cpp
@@ -12,11 +12,21 @@ bool FCrashReportAnalytics::bIsInitialized;
 TSharedPtr<IAnalyticsProviderET> FCrashReportAnalytics::Analytics;
 
 /**
- * Default config func that essentially tells the crash reporter to disable analytics.
+ * Default config func that essentially tells the crash reporter to disable analytics, unless
+ * environment variables are set.
  */
 FAnalyticsET::Config DefaultAnalyticsConfigFunc()
 {
-	return FAnalyticsET::Config();
+	FAnalyticsET::Config Config = FAnalyticsET::Config();
+	// Check environment for settings. Both url and key needs to be set.
+	FString Url = FPlatformMisc::GetEnvironmentVariable(TEXT("UE_CRC_TELEMETRY_URL"));
+	FString Key = FPlatformMisc::GetEnvironmentVariable(TEXT("UE_CRC_TELEMETRY_KEY"));
+	if (!Url.IsEmpty() && !Key.IsEmpty())
+	{
+		Config.APIServerET = Url;
+		Config.APIKeyET = Key;
+	}
+	return Config;
 }
 
 /**
@@ -44,22 +54,12 @@ void FCrashReportAnalytics::Initialize(const FString& EpicAccountId)
 	checkf(!bIsInitialized, TEXT("FCrashReportAnalytics::Initialize called more than once."));
 
 	// Allow build machines to force CRC to enable internal telemetry.
-	#if defined(CRC_TELEMETRY_URL) && defined(CRC_TELEMETRY_KEY_DEV) && defined(CRC_TELEMETRY_KEY_RELEASE)
-
-		// We always use the "Release" analytics account unless we're running in analytics test mode (usually with
-		// a command-line parameter), or we're an internal Epic build
-		bool bUseReleaseAccount =
-			(GetAnalyticsBuildType() == EAnalyticsBuildType::Development || GetAnalyticsBuildType() == EAnalyticsBuildType::Release) &&
-			!FEngineBuildSettings::IsInternalBuild();	// Internal Epic build
-
+	#if defined(CRC_TELEMETRY_URL) && defined(CRC_TELEMETRY_KEY) 
 		FAnalyticsET::Config Config;
 		Config.APIServerET = TEXT(CRC_TELEMETRY_URL);
-		Config.APIKeyET = bUseReleaseAccount ? TEXT(CRC_TELEMETRY_KEY_RELEASE) : TEXT(CRC_TELEMETRY_KEY_DEV);
-
+		Config.APIKeyET = TEXT(CRC_TELEMETRY_KEY);
 	#else
-
 		FAnalyticsET::Config Config = GetCrashReportAnalyticsConfigFunc()();
-
 	#endif
 
 	if (!Config.APIServerET.IsEmpty())
diff --git a/Engine/Source/Runtime/CrashReportCore/Private/CrashReportCoreConfig.cpp b/Engine/Source/Runtime/CrashReportCore/Private/CrashReportCoreConfig.cpp
--- a/Engine/Source/Runtime/CrashReportCore/Private/CrashReportCoreConfig.cpp
+++ b/Engine/Source/Runtime/CrashReportCore/Private/CrashReportCoreConfig.cpp
@@ -31,8 +31,8 @@ FCrashReportCoreConfig::FCrashReportCoreConfig()
 
 	if (!GConfig->GetString(*SectionName, TEXT("DataRouterUrl"), DataRouterUrl, GEngineIni))
 	{
-#if defined CRC_DATAROUTER_FALLBACK
-		DataRouterUrl = TEXT(CRC_DATAROUTER_FALLBACK);
+#if defined CRC_DEFAULT_URL
+		DataRouterUrl = TEXT(CRC_DEFAULT_URL);
 #else
 		DataRouterUrl = TEXT("");
 #endif
@@ -88,8 +88,8 @@ FCrashReportCoreConfig::FCrashReportCoreConfig()
 
 	if (!GConfig->GetString(*SectionName, TEXT("CompanyName"), CompanyName, GEngineIni))
 	{
-#if defined(CRC_COMPANY_NAME_FALLBACK)
-		CompanyName = TEXT(CRC_COMPANY_NAME_FALLBACK);
+#if defined(CRC_DEFAULT_COMPANY_NAME)
+		CompanyName = TEXT(CRC_DEFAULT_COMPANY_NAME);
 #else
 		CompanyName = TEXT("");
 #endif
