Fix deprecation warning

diff --git a/Engine/Source/Programs/CrashReportClient/Private/CrashReportAnalyticsSessionSummary.cpp b/Engine/Source/Programs/CrashReportClient/Private/CrashReportAnalyticsSessionSummary.cpp
--- a/Engine/Source/Programs/CrashReportClient/Private/CrashReportAnalyticsSessionSummary.cpp
+++ b/Engine/Source/Programs/CrashReportClient/Private/CrashReportAnalyticsSessionSummary.cpp
@@ -467,7 +467,7 @@ void FCrashReportAnalyticsSessionSummary::LogEvent(const TCHAR* Event, bool bFor
 	{
 		FScopeLock ScopedLock(&LoggerLock);
 		TGuardValue<bool> ReentrantGuard(bLoggerReentrantGuard, true);
-		if (*ReentrantGuard) // Read the old value.
+		if (ReentrantGuard.GetOriginalValue()) // Read the old value.
 		{
 			return; // Prevent renentrant logging.
 		}
