WORKAROUND CL33938220.
The module name is ok but the CrashReportClientEditor target is causing a warning that can be safely ignored.

diff --git a/Engine/Source/Programs/CrashReportClient/Private/CrashReportClientApp.cpp b/Engine/Source/Programs/CrashReportClient/Private/CrashReportClientApp.cpp
--- a/Engine/Source/Programs/CrashReportClient/Private/CrashReportClientApp.cpp
+++ b/Engine/Source/Programs/CrashReportClient/Private/CrashReportClientApp.cpp
@@ -62,7 +62,26 @@ const float IdealTickRate = 30.f;
 /** Set this to true in the code to open the widget reflector to debug the UI */
 const bool RunWidgetReflector = false;
 
+//WORKAROUND CL33938220. The module name is ok but the CrashReportClientEditor target is causing a warning that can be safely ignored. 
+#ifdef __clang__
+#pragma clang diagnostic push
+#pragma clang diagnostic ignored "-Wdeprecated-declarations"
+#elif defined(__GNUC__)
+#pragma GCC diagnostic push
+#pragma GCC diagnostic ignored "-Wdeprecated-declarations"
+#else
+#pragma warning(push)
+#pragma warning(disable : 4996) // 'function' was declared deprecated
+#endif
 IMPLEMENT_APPLICATION(CrashReportClient, "CrashReportClient");
+#ifdef __clang__
+#pragma clang diagnostic pop
+#elif defined(__GNUC__)
+#pragma GCC diagnostic pop
+#else
+#pragma warning(pop)
+#endif
+
 DEFINE_LOG_CATEGORY(CrashReportClientLog);
 
 /** Directory containing the report */
