[Crash Report Client] Handle optional attachments (only screenshot for now)

This change introduces support for optional attachments in the Crash Report Client. Optional attachments are extra files stored under the "OptionalAttachments" subdirectory within a crash report folder, and are never uploaded by default. Even in the "implicit send" mode, the user has to explicitly indicate they want to upload these optional attachments. If the user indicates they want to upload the optional attachments, the Crash Report Client will wait for the base report to have finished uploading, and then send an additional HTTP request to include the optional attachments.

For now, only one type of optional attachment is supported: a crash screenshot. The screenshot is shown in the crash report dialog so that the user can see what the screenshot looks like, if they choose to upload it. See CL 39533466 for the changes on the engine side and how the optional attachments are produced.

Only the DataRouter uploader path supports optional attachments upload.

#tests tested extensively locally and with our crash ingestion system, which now supports adding attachments to existing reports.
#localization

diff --git a/Engine/Source/Programs/CrashReportClient/Private/CrashReportClient.cpp b/Engine/Source/Programs/CrashReportClient/Private/CrashReportClient.cpp
--- a/Engine/Source/Programs/CrashReportClient/Private/CrashReportClient.cpp
+++ b/Engine/Source/Programs/CrashReportClient/Private/CrashReportClient.cpp
@@ -44,6 +44,7 @@ FCrashReportClient::FCrashReportClient(const FPlatformErrorReport& InErrorReport
 	, DataRouterUploader(FCrashReportCoreConfig::Get().GetDataRouterURL(), FCrashReportCoreConfig::Get().GetDataRouterUnixSocket())
 	, bShouldWindowBeHidden(false)
 	, bSendData(bImplicitSend)
+	, bSendOptionalAttachments(false)
 	, bIsSuccesfullRestart(false)
 	, bIsUploadComplete(false)
 {
@@ -131,10 +132,23 @@ FReply FCrashReportClient::CopyFilesToClipboard()
 }
 #endif
 
-FReply FCrashReportClient::Submit()
+FReply FCrashReportClient::Submit(bool bIncludeOptionalAttachments)
 {
 	bSendData = true;
 	StoreCommentAndUpload();
+
+	if (bIncludeOptionalAttachments)
+	{
+		bSendOptionalAttachments = true;
+	}
+
+	bShouldWindowBeHidden = true;
+	return FReply::Handled();
+}
+
+FReply FCrashReportClient::SubmitOptionalAttachmentsAndClose()
+{
+	bSendOptionalAttachments = true;
 	bShouldWindowBeHidden = true;
 	return FReply::Handled();
 }
@@ -329,6 +343,33 @@ bool FCrashReportClient::Tick(float UnusedDeltaTime)
 		}
 	}
 
+	// Optional attachments are only supported by the Data Router uploader.
+	if (bSendData && bSendOptionalAttachments && DataRouterUploader.IsEnabled())
+	{
+		UE_LOG(CrashReportClientLog, Log, TEXT("User has requested to upload optional attachments"));
+
+		// Wait until the base report has been sent (successfully or not).
+		if (!DataRouterUploader.IsFinished())
+		{
+			return true;
+		}
+
+		// Only send optional attachments if the base report was uploaded successfully.
+		if (DataRouterUploader.IsFinishedSuccessfully())
+		{
+			// Start upload and keep ticking until done.
+			DataRouterUploader.BeginUploadOptionalAttachments(ErrorReport);
+
+			bSendOptionalAttachments = false;
+			return true;
+		}
+		else
+		{
+			UE_LOG(CrashReportClientLog, Warning, TEXT("User has requested to upload optional attachments but the base report wasn't uploaded successfully: cancelling optional attachments upload"));
+			bSendOptionalAttachments = false;
+		}
+	}
+
 	if (!bShouldWindowBeHidden)
 	{
 		return true;
diff --git a/Engine/Source/Programs/CrashReportClient/Private/CrashReportClient.h b/Engine/Source/Programs/CrashReportClient/Private/CrashReportClient.h
--- a/Engine/Source/Programs/CrashReportClient/Private/CrashReportClient.h
+++ b/Engine/Source/Programs/CrashReportClient/Private/CrashReportClient.h
@@ -11,6 +11,8 @@
 #include "CrashUpload.h"
 #include "Containers/Ticker.h"
 
+DECLARE_LOG_CATEGORY_EXTERN(CrashReportClientLog, Log, All);
+
 #if !CRASH_REPORT_UNATTENDED_ONLY
 
 #include "Input/Reply.h"
@@ -86,7 +88,13 @@ public:
 	 * Respond to the user pressing Submit
 	 * @return Whether the request was handled
 	 */
-	FReply Submit();
+	FReply Submit(bool bIncludeOptionalAttachments = false);
+
+	/**
+	 * Respond to the user pressing Submit for optional attachments
+	 * @return Whether the request was handled
+	 */
+	FReply SubmitOptionalAttachmentsAndClose();
 
 	/**
 	* Respond to the user pressing Submit and Restart
@@ -200,6 +208,9 @@ private:
 	/** Whether we send the data. */
 	bool bSendData;
 
+	/** Whether to send optional attachments, after the base crash report has been successfully sent. */
+	bool bSendOptionalAttachments;
+
 	/** Store if user has elected to restart crashing process, and if it was a success.*/
 	bool bIsSuccesfullRestart;
 
diff --git a/Engine/Source/Programs/CrashReportClient/Private/CrashReportClientApp.cpp b/Engine/Source/Programs/CrashReportClient/Private/CrashReportClientApp.cpp
--- a/Engine/Source/Programs/CrashReportClient/Private/CrashReportClientApp.cpp
+++ b/Engine/Source/Programs/CrashReportClient/Private/CrashReportClientApp.cpp
@@ -57,6 +57,12 @@ const FVector2D InitialWindowDimensions(740, 560);
 /** Simple dialog window size */
 const FVector2D InitialSimpleWindowDimensions(740, 300);
 
+/** Default main window size with screenshot present */
+const FVector2D InitialWindowDimensionsWithScreenshot(740, 1000);
+
+/** Simple dialog window size with screenshot present */
+const FVector2D InitialSimpleWindowDimensionsWithScreenshot(920, 700);
+
 /** Average tick rate the app aims for */
 const float IdealTickRate = 30.f;
 
@@ -295,7 +301,7 @@ SubmitCrashReportResult RunWithUI(FPlatformErrorReport ErrorReport, bool bImplic
 	// Try to initialize the renderer. It's possible that we launched when the driver crashed so try a few times before giving up.
 	bool bRendererInitialized = false;
 	bool bRendererFailedToInitializeAtLeastOnce = false;
-	do 
+	do
 	{
 		SlateRendererInitRetryCount--;
 		bRendererInitialized = FSlateApplication::Get().InitializeRenderer(SlateRenderer, true);
@@ -329,8 +335,20 @@ SubmitCrashReportResult RunWithUI(FPlatformErrorReport ErrorReport, bool bImplic
 	// set the normal IsEngineExitRequested() when outer frame is closed
 	FSlateApplication::Get().SetExitRequestedHandler(FSimpleDelegate::CreateStatic(&OnRequestExit));
 
-	// Prepare the custom Slate styles
-	FCrashReportClientStyle::Initialize();
+	// Prepare the custom Slate styles.
+	// If there's an optional screenshot attachment, ensure Slate loads it so that we can show it as part of the dialog.
+	FCrashReportClientStyleOptions StyleOptions = {};
+	FString CrashScreenshotPath = FPaths::Combine(ErrorReport.GetReportDirectory(), FGenericCrashContext::CrashReportOptionalAttachmentsDirName, TEXT("CrashCapture.jpg"));
+	bool bHasScreenshot = FPaths::FileExists(CrashScreenshotPath);
+	FVector2D CrashScreenshotSize{ 0, 0 };
+	if (bHasScreenshot && FPrimaryCrashProperties::Get()->IsValid())
+	{
+		StyleOptions.OptionalCrashScreenshot.Filepath = CrashScreenshotPath;
+		StyleOptions.OptionalCrashScreenshot.ImageSize = FVector2D(FPrimaryCrashProperties::Get()->CrashScreenshotWidth, FPrimaryCrashProperties::Get()->CrashScreenshotHeight);
+
+		bHasScreenshot = !StyleOptions.OptionalCrashScreenshot.ImageSize.IsZero();
+	}
+	FCrashReportClientStyle::Initialize(StyleOptions);
 
 	// Create the main implementation object
 	TSharedRef<FCrashReportClient> CrashReportClient = MakeShared<FCrashReportClient>(ErrorReport, bImplicitSend);
@@ -362,11 +380,21 @@ SubmitCrashReportResult RunWithUI(FPlatformErrorReport ErrorReport, bool bImplic
 		FText::Format(NSLOCTEXT("CrashReportClient", "CrashReportClientAppName", "Unreal Engine {0} Crash Reporter"), MajorEngineVersion) :
 		CrashedAppText;
 
+	FVector2D WindowDimensions;
+	if (bImplicitSend)
+	{
+		WindowDimensions = bHasScreenshot ? InitialSimpleWindowDimensionsWithScreenshot : InitialSimpleWindowDimensions;
+	}
+	else
+	{
+		WindowDimensions = bHasScreenshot ? InitialWindowDimensionsWithScreenshot : InitialWindowDimensions;
+	}
+
 	TSharedRef<SWindow> Window = FSlateApplication::Get().AddWindow(
 		SNew(SWindow)
 		.Title(WindowTitle)
 		.HasCloseButton(FCrashReportCoreConfig::Get().IsAllowedToCloseWithoutSending())
-		.ClientSize(bImplicitSend ? InitialSimpleWindowDimensions : InitialWindowDimensions)
+		.ClientSize(WindowDimensions)
 		[
 			ClientControl
 		]);
diff --git a/Engine/Source/Programs/CrashReportClient/Private/CrashReportClientStyle.cpp b/Engine/Source/Programs/CrashReportClient/Private/CrashReportClientStyle.cpp
--- a/Engine/Source/Programs/CrashReportClient/Private/CrashReportClientStyle.cpp
+++ b/Engine/Source/Programs/CrashReportClient/Private/CrashReportClientStyle.cpp
@@ -20,11 +20,11 @@
 
 TSharedPtr< FSlateStyleSet > FCrashReportClientStyle::StyleSet = nullptr;
 
-void FCrashReportClientStyle::Initialize()
+void FCrashReportClientStyle::Initialize(const FCrashReportClientStyleOptions& Options)
 {
 	if(!StyleSet.IsValid())
 	{
-		StyleSet = Create();
+		StyleSet = Create(Options);
 		FSlateStyleRegistry::RegisterSlateStyle(*StyleSet);
 	}
 }
@@ -45,7 +45,7 @@ namespace
 	}
 }
 
-TSharedRef< FSlateStyleSet > FCrashReportClientStyle::Create()
+TSharedRef< FSlateStyleSet > FCrashReportClientStyle::Create(const FCrashReportClientStyleOptions& Options)
 {
 	TSharedRef<FSlateStyleSet> StyleRef = MakeShareable(new FSlateStyleSet("CrashReportClientStyle"));
 	FSlateStyleSet& Style = StyleRef.Get();
@@ -73,6 +73,14 @@ TSharedRef< FSlateStyleSet > FCrashReportClientStyle::Create()
 	const FVector2D Icon16x16( 16.0f, 16.0f );
 	FSlateBrush* GenericWhiteBox = new IMAGE_BRUSH( "Old/White", Icon16x16 );
 
+	// Load the screenshot image, if one is specified.
+	const FString& ScreenshotPath = Options.OptionalCrashScreenshot.Filepath;
+	bool bHasValidScreenshotFile = FPaths::FileExists(ScreenshotPath);
+	if (bHasValidScreenshotFile && !Options.OptionalCrashScreenshot.ImageSize.IsZero())
+	{
+		Style.Set("CrashScreenshot", new FSlateImageBrush(ScreenshotPath, Options.OptionalCrashScreenshot.ImageSize));
+	}
+
 	// Scrollbar
 	const FScrollBarStyle ScrollBar = FScrollBarStyle()
 		.SetVerticalTopSlotImage(IMAGE_BRUSH("Common/Scrollbar_Background_Vertical", FVector2D(8, 8)))
diff --git a/Engine/Source/Programs/CrashReportClient/Private/CrashReportClientStyle.h b/Engine/Source/Programs/CrashReportClient/Private/CrashReportClientStyle.h
--- a/Engine/Source/Programs/CrashReportClient/Private/CrashReportClientStyle.h
+++ b/Engine/Source/Programs/CrashReportClient/Private/CrashReportClientStyle.h
@@ -9,6 +9,13 @@
 
 #include "Styling/SlateStyle.h"
 
+struct FCrashReportClientStyleOptions {
+	struct {
+		FString Filepath;
+		FVector2D ImageSize = { 0,0 };
+	} OptionalCrashScreenshot;
+};
+
 /** Slate styles for the crash report client app */
 class FCrashReportClientStyle
 {
@@ -16,7 +23,7 @@ public:
 	/**
 	 * Set up specific styles for the crash report client app
 	 */
-	static void Initialize();
+	static void Initialize(const FCrashReportClientStyleOptions& Options);
 
 	/**
 	 * Tidy up on shut-down
@@ -29,7 +36,7 @@ public:
 	static const ISlateStyle& Get();
 
 private:
-	static TSharedRef< FSlateStyleSet > Create();
+	static TSharedRef<FSlateStyleSet> Create(const FCrashReportClientStyleOptions& Options);
 
 	/** Singleton style object */
 	static TSharedPtr<FSlateStyleSet> StyleSet;
diff --git a/Engine/Source/Programs/CrashReportClient/Private/SCrashReportClient.cpp b/Engine/Source/Programs/CrashReportClient/Private/SCrashReportClient.cpp
--- a/Engine/Source/Programs/CrashReportClient/Private/SCrashReportClient.cpp
+++ b/Engine/Source/Programs/CrashReportClient/Private/SCrashReportClient.cpp
@@ -14,6 +14,7 @@
 #include "Widgets/Text/STextBlock.h"
 #include "Widgets/Text/SRichTextBlock.h"
 #include "Widgets/Layout/SSplitter.h"
+#include "Widgets/Layout/SScaleBox.h"
 #include "Widgets/Colors/SColorBlock.h"
 #include "Widgets/Input/SCheckBox.h"
 #include "Widgets/Input/SButton.h"
@@ -82,6 +83,63 @@ void SCrashReportClient::ConstructDetailedDialog(const TSharedRef<FCrashReportCl
 		LOCTEXT( "CrashReportData", "Crash reports comprise diagnostics files (<a id=\"browser\" href=\"{0}\" style=\"Richtext.Hyperlink\">click here to view directory</>) and the following summary information: " ),
 		FText::FromString( CrashReportClient->GetCrashDirectory()) );
 
+	const FText IncludeScreenshotText = LOCTEXT("CrashReportIncludeScreenshot", "Include screenshot in the crash report");
+
+	const FSlateBrush* Screenshot = FCrashReportClientStyle::Get().GetOptionalBrush("CrashScreenshot", nullptr, nullptr);
+
+	SVerticalBox::FSlot::FSlotArguments ScreenshotSlot = SVerticalBox::Slot();
+	if (Screenshot)
+	{
+		ScreenshotSlot.Padding(FMargin(4, 10))
+		.MaxHeight(500)
+			[
+				SNew(SVerticalBox)
+				+SVerticalBox::Slot()
+					.AutoHeight()
+					.Padding(4)
+					[
+						SNew(SHorizontalBox)
+							+ SHorizontalBox::Slot()
+							.AutoWidth()
+							[
+								SNew(SCheckBox)
+									.IsChecked_Lambda([this]()
+										{
+											return bIncludeScreenshotInCrashReport ? ECheckBoxState::Checked : ECheckBoxState::Unchecked;
+										})
+									.OnCheckStateChanged_Lambda([this](ECheckBoxState InCheckBoxState)
+										{
+											bIncludeScreenshotInCrashReport = InCheckBoxState == ECheckBoxState::Checked;
+										})
+							]
+							+ SHorizontalBox::Slot()
+							.AutoWidth()
+							.Padding(4, 1,1,1)
+							[
+								SNew(SRichTextBlock)
+									.Text(IncludeScreenshotText)
+							]
+					]
+					+ SVerticalBox::Slot()
+					[
+						SNew(SScaleBox)
+							.Stretch(EStretch::ScaleToFit)
+							[
+								SNew(SImage)
+									.IsEnabled(Screenshot != nullptr)
+									.Image(Screenshot)
+									.ColorAndOpacity_Lambda([this]() {
+										return FLinearColor(1.0f, 1.0f, 1.0f, bIncludeScreenshotInCrashReport ? 1.0f : 0.4f);
+									})
+							]
+					]
+			];
+	}
+	else
+	{
+		ScreenshotSlot.MaxHeight(0);
+	}
+
 	ChildSlot
 	[
 		SNew(SBorder)
@@ -122,6 +180,8 @@ void SCrashReportClient::ConstructDetailedDialog(const TSharedRef<FCrashReportCl
 				+ SRichTextBlock::HyperlinkDecorator( TEXT("browser"), FSlateHyperlinkRun::FOnClick::CreateStatic( &OnBrowserLinkClicked ) )
 			]
 
+			+ScreenshotSlot
+
 			+SVerticalBox::Slot()
 			.Padding(FMargin(4, 10, 4, 4))
 			[
@@ -319,7 +379,7 @@ void SCrashReportClient::ConstructDetailedDialog(const TSharedRef<FCrashReportCl
 					SNew(SButton)
 					.ContentPadding( FMargin(8,2) )
 					.Text(LOCTEXT("Send", "Send and Close"))
-					.OnClicked(Client, &FCrashReportClient::Submit)
+					.OnClicked(Client, &FCrashReportClient::Submit, bIncludeScreenshotInCrashReport)
 					.IsEnabled(this, &SCrashReportClient::IsSendEnabled)
 					.ToolTipText_Static(&SCrashReportClient::GetSendTooltip)
 				]
@@ -356,6 +416,45 @@ void SCrashReportClient::ConstructSimpleDialog(const TSharedRef<FCrashReportClie
 	}
 	CrashedAppName.RemoveFromEnd(TEXT("Game"));
 
+	// This simple dialog is used for the unattended mode where the base crash report is sent
+	// unconditionally, but we still show options for optional attachments such as the screenshot.
+	const FText SendScreenshotText = LOCTEXT("CrashReportSendScreenshot", "Send screenshot");
+	const FText ScreenshotDescriptionText = LOCTEXT("CrashReportScreenshotDescription",
+		"The following screenshot was captured when the application crashed. You can optionally include this screenshot as part of the crash report.");
+	const FSlateBrush* Screenshot = FCrashReportClientStyle::Get().GetOptionalBrush("CrashScreenshot", nullptr, nullptr);
+
+	SVerticalBox::FSlot::FSlotArguments ScreenshotDescriptionSlot = SVerticalBox::Slot();
+	SVerticalBox::FSlot::FSlotArguments ScreenshotSlot = SVerticalBox::Slot();
+	if (Screenshot)
+	{
+		ScreenshotDescriptionSlot
+		.AutoHeight()
+		.Padding(FMargin(4, 10))
+		[
+			SNew(SRichTextBlock)
+			.Text(ScreenshotDescriptionText)
+			.AutoWrapText(true)
+		];
+
+		ScreenshotSlot
+		.Padding(FMargin(4, 4))
+		.MaxHeight(500)
+		.AutoHeight()
+		[
+			SNew(SScaleBox)
+			.Stretch(EStretch::ScaleToFit)
+			[
+				SNew(SImage)
+				.Image(Screenshot)
+			]
+		];
+	}
+	else
+	{
+		ScreenshotDescriptionSlot.MaxHeight(0);
+		ScreenshotSlot.MaxHeight(0);
+	}
+
 	ChildSlot
 	[
 		SNew(SBorder)
@@ -388,20 +487,48 @@ void SCrashReportClient::ConstructSimpleDialog(const TSharedRef<FCrashReportClie
 				+ SRichTextBlock::HyperlinkDecorator( TEXT("browser"), FSlateHyperlinkRun::FOnClick::CreateStatic( &OnBrowserLinkClicked ) )
 			]
 
+			+ScreenshotDescriptionSlot
+			+ScreenshotSlot
+
 			// Stuff anchored to the bottom
 			+SVerticalBox::Slot()
-			.Padding( FMargin(4, 4+16, 4, 4) )
+			.Padding( FMargin(4, 4, 4, 4) )
 			[
 				SNew(SHorizontalBox)
 
 				+SHorizontalBox::Slot()
-				.HAlign(HAlign_Right)
+				.FillWidth(1.0f)
+				.HAlign(HAlign_Left)
+				.VAlign(VAlign_Bottom)
+				.Padding(0)
+				[
+					SNew(SSpacer)
+				]
+
+				+SHorizontalBox::Slot()
+				.HAlign(HAlign_Left)
+				.VAlign(VAlign_Bottom)
+				.AutoWidth()
+				.Padding( FMargin(6) )
+				[
+					SNew(SButton)
+					.ContentPadding(FMargin(8, 2))
+					.Text(SendScreenshotText)
+					.Visibility(Screenshot ? EVisibility::Visible : EVisibility::Collapsed)
+					.OnClicked_Lambda([Client]() {
+						return Client->SubmitOptionalAttachmentsAndClose();
+					 })
+				]
+
+				+SHorizontalBox::Slot()
+				.HAlign(HAlign_Left)
 				.VAlign(VAlign_Bottom)
+				.AutoWidth()
 				.Padding( FMargin(6) )
 				[
 					SNew(SButton)
 					.ContentPadding( FMargin(8,2) )
-					.Text(LOCTEXT("Close", "Close"))
+					.Text(Screenshot ? LOCTEXT("CloseWithoutSending", "Close Without Sending") : LOCTEXT("Close", "Close"))
 					.OnClicked(Client, &FCrashReportClient::Close)
 				]		
 			]
diff --git a/Engine/Source/Programs/CrashReportClient/Private/SCrashReportClient.h b/Engine/Source/Programs/CrashReportClient/Private/SCrashReportClient.h
--- a/Engine/Source/Programs/CrashReportClient/Private/SCrashReportClient.h
+++ b/Engine/Source/Programs/CrashReportClient/Private/SCrashReportClient.h
@@ -86,6 +86,8 @@ private:
 
 	bool bHasUserCommentErrors;
 	bool bHideSubmitAndRestart;
+
+	bool bIncludeScreenshotInCrashReport = false;
 };
 
 #endif // !CRASH_REPORT_UNATTENDED_ONLY
diff --git a/Engine/Source/Programs/CrashReportClient/Public/CrashReportClientApp.h b/Engine/Source/Programs/CrashReportClient/Public/CrashReportClientApp.h
--- a/Engine/Source/Programs/CrashReportClient/Public/CrashReportClientApp.h
+++ b/Engine/Source/Programs/CrashReportClient/Public/CrashReportClientApp.h
@@ -11,8 +11,6 @@
 	#include "StandaloneRenderer.h"
 #endif // CRASH_REPORT_UNATTENDED_ONLY
 
-DECLARE_LOG_CATEGORY_EXTERN(CrashReportClientLog, Log, All)
-
 /**
  * Run the crash report client app
  */
diff --git a/Engine/Source/Runtime/CrashReportCore/Private/CrashDescription.cpp b/Engine/Source/Runtime/CrashReportCore/Private/CrashDescription.cpp
--- a/Engine/Source/Runtime/CrashReportCore/Private/CrashDescription.cpp
+++ b/Engine/Source/Runtime/CrashReportCore/Private/CrashDescription.cpp
@@ -579,6 +579,8 @@ void FCrashContext::SetupPrimaryCrashProperties()
 		GetCrashProperty( bLowMemoryWarning, FGenericCrashContext::GameDataTag, TEXT("bLowMemoryCalled"));
 		GetCrashProperty( bInBackground, FGenericCrashContext::GameDataTag, TEXT("bInBackground"));
 		GetCrashProperty( bIsRequestingExit, FGenericCrashContext::GameDataTag, TEXT("IsRequestingExit"));
+		GetCrashProperty( CrashScreenshotWidth, FGenericCrashContext::EngineDataTag, TEXT("CrashScreenshotWidth"));
+		GetCrashProperty( CrashScreenshotHeight, FGenericCrashContext::EngineDataTag, TEXT("CrashScreenshotHeight"));
 
 		FString CPUBrandString;
 		GetCrashProperty( CPUBrandString, FGenericCrashContext::RuntimePropertiesTag, TEXT("Misc.CPUBrand"));
diff --git a/Engine/Source/Runtime/CrashReportCore/Private/CrashUpload.cpp b/Engine/Source/Runtime/CrashReportCore/Private/CrashUpload.cpp
--- a/Engine/Source/Runtime/CrashReportCore/Private/CrashUpload.cpp
+++ b/Engine/Source/Runtime/CrashReportCore/Private/CrashUpload.cpp
@@ -768,6 +768,21 @@ void FCrashUploadToDataRouter::BeginUpload(const FPlatformErrorReport& PlatformE
 	}
 }
 
+void FCrashUploadToDataRouter::BeginUploadOptionalAttachments(const FPlatformErrorReport& PlatformErrorReport)
+{
+	ErrorReport = PlatformErrorReport;
+	PendingFiles = FPlatformErrorReport(ErrorReport.GetReportDirectory()).GetOptionalAttachmentsToUpload();
+	UE_LOG(CrashReportCoreLog, Log, TEXT("Got %d pending optional attachment files to upload"), PendingFiles.Num());
+
+	PauseState = EUploadState::Finished;
+	// We should only upload optional attachments when the base report was successfully uploaded (EUploadState::Finished).
+	if (State == EUploadState::Finished)
+	{
+		SetCurrentState(EUploadState::CompressAndSendData);
+		CompressAndSendData();
+	}
+}
+
 void FCrashUploadToDataRouter::CompressAndSendData()
 {
 	FCompressedHeader CompressedHeader;
diff --git a/Engine/Source/Runtime/CrashReportCore/Private/GenericErrorReport.cpp b/Engine/Source/Runtime/CrashReportCore/Private/GenericErrorReport.cpp
--- a/Engine/Source/Runtime/CrashReportCore/Private/GenericErrorReport.cpp
+++ b/Engine/Source/Runtime/CrashReportCore/Private/GenericErrorReport.cpp
@@ -42,6 +42,7 @@ namespace
 
 FGenericErrorReport::FGenericErrorReport(const FString& Directory)
 	: ReportDirectory(Directory)
+	, OptionalAttachmentsDirectory(FPaths::Combine(Directory, FGenericCrashContext::CrashReportOptionalAttachmentsDirName))
 	, bValidCallstack(true)
 {
 	auto FilenamesVisitor = MakeDirectoryVisitor([this](const TCHAR* FilenameOrDirectory, bool bIsDirectory) {
@@ -51,7 +52,17 @@ FGenericErrorReport::FGenericErrorReport(const FString& Directory)
 		}
 		return true;
 	});
+
+	auto OptionalFilenamesVisitor = MakeDirectoryVisitor([this](const TCHAR* FilenameOrDirectory, bool bIsDirectory) {
+		if (!bIsDirectory)
+		{
+			OptionalAttachmentsFilenames.Push(FPaths::GetCleanFilename(FilenameOrDirectory));
+		}
+		return true;
+	});
+
 	FPlatformFileManager::Get().GetPlatformFile().IterateDirectory(*ReportDirectory, FilenamesVisitor);
+	FPlatformFileManager::Get().GetPlatformFile().IterateDirectory(*OptionalAttachmentsDirectory, OptionalFilenamesVisitor);
 }
 
 bool FGenericErrorReport::SetUserComment(const FText& UserComment)
@@ -238,6 +249,17 @@ TArray<FString> FGenericErrorReport::GetFilesToUpload() const
 	return FilesToUpload;
 }
 
+TArray<FString> FGenericErrorReport::GetOptionalAttachmentsToUpload() const
+{
+	TArray<FString> FilesToUpload;
+
+	for (const auto& Filename : OptionalAttachmentsFilenames)
+	{
+		FilesToUpload.Push(OptionalAttachmentsDirectory / Filename);
+	}
+	return FilesToUpload;
+}
+
 void FGenericErrorReport::DeleteFiles()
 {
 	for (const auto& Filename: ReportFilenames)
diff --git a/Engine/Source/Runtime/CrashReportCore/Public/CrashDescription.h b/Engine/Source/Runtime/CrashReportCore/Public/CrashDescription.h
--- a/Engine/Source/Runtime/CrashReportCore/Public/CrashDescription.h
+++ b/Engine/Source/Runtime/CrashReportCore/Public/CrashDescription.h
@@ -473,6 +473,12 @@ struct FPrimaryCrashProperties
 	 */
 	bool bIsRequestingExit;
 
+	/**
+	 * Width and height of the crash screenshot, if present
+	 */
+	uint32 CrashScreenshotWidth = 0;
+	uint32 CrashScreenshotHeight = 0;
+
 protected:
 	/** Default constructor. */
 	FPrimaryCrashProperties();
diff --git a/Engine/Source/Runtime/CrashReportCore/Public/CrashUpload.h b/Engine/Source/Runtime/CrashReportCore/Public/CrashUpload.h
--- a/Engine/Source/Runtime/CrashReportCore/Public/CrashUpload.h
+++ b/Engine/Source/Runtime/CrashReportCore/Public/CrashUpload.h
@@ -46,6 +46,14 @@ public:
 		return State >= EUploadState::FirstCompletedState;
 	}
 
+	/**
+	 * Determine whether the upload has finished successfully.
+	 */
+	bool IsFinishedSuccessfully() const
+	{
+		return State == EUploadState::Finished;
+	}
+
 	void Cancel()
 	{
 		SetCurrentState(EUploadState::Cancelled);
@@ -260,6 +268,12 @@ public:
 
 	void BeginUpload(const FPlatformErrorReport& PlatformErrorReport);
 
+	/**
+	 * Compresses all optional attachments files and sends one compressed file.
+	 * Requires the base report to have been already sent (via BeginUpload).
+	 */
+	void BeginUploadOptionalAttachments(const FPlatformErrorReport& PlatformErrorReport);
+
 	/**
 	 * Compresses all crash report files and sends one compressed file.
 	 */
diff --git a/Engine/Source/Runtime/CrashReportCore/Public/GenericErrorReport.h b/Engine/Source/Runtime/CrashReportCore/Public/GenericErrorReport.h
--- a/Engine/Source/Runtime/CrashReportCore/Public/GenericErrorReport.h
+++ b/Engine/Source/Runtime/CrashReportCore/Public/GenericErrorReport.h
@@ -65,6 +65,12 @@ public:
 	 */
 	TArray<FString> GetFilesToUpload() const;
 
+	/**
+	 * Provide full paths to all the report's optional attachment files
+	 * @return List of paths
+	 */
+	TArray<FString> GetOptionalAttachmentsToUpload() const;
+
 	/**
 	 * Load the WER XML file for this report
 	 * @note This is Windows specific and so shouldn't really be part of the public interface, but currently the server
@@ -133,6 +139,12 @@ protected:
 	/** List of leaf filenames of all the files in the report folder */
 	TArray<FString> ReportFilenames;
 
+	/** Full path to the directory the report's optional attachment files are in */
+	FString OptionalAttachmentsDirectory;
+
+	/** List of leaf filenames of all the files in the report's optional attachments folder */
+	TArray<FString> OptionalAttachmentsFilenames;
+
 	/** Whether the error report generated an valid callstack. */
 	mutable bool bValidCallstack;
 };
