Making sure that any existing game and engine data get read back in by the CRC app when loading temp crash contexts.

diff --git a/Engine/Source/Programs/CrashReportClient/Private/CrashReportClientApp.cpp b/Engine/Source/Programs/CrashReportClient/Private/CrashReportClientApp.cpp
--- a/Engine/Source/Programs/CrashReportClient/Private/CrashReportClientApp.cpp
+++ b/Engine/Source/Programs/CrashReportClient/Private/CrashReportClientApp.cpp
@@ -998,6 +998,84 @@ static bool LoadTempCrashContextFromFile(FSharedCrashContext& CrashContext, uint
 	FindAndParseValue(ContextProperties, TEXT("SendUsageData"), UserSettings.bSendUsageData);
 	FindAndCopyValue(ContextProperties, TEXT("LogFilePath"), UserSettings.LogFilePath);
 
+	// dynamic data
+	
+	FMemory::Memset(CrashContext.DynamicData, 0);
+	TCHAR* DynamicDataStart = &CrashContext.DynamicData[0];
+	TCHAR* DynamicDataPtr = DynamicDataStart;
+	
+	// -1 to allow space for null terminator
+	#define CR_DYNAMIC_BUFFER_REMAIN uint32((CR_MAX_DYNAMIC_BUFFER_CHARS) - (DynamicDataPtr-DynamicDataStart) - 1)
+
+	// When encoding the plugins list and engine and game data key/value pairs into the dynamic data segment
+	// we use 1 and 2 to denote delimiter and equals respectively. This is necessary since the values could 
+	// contain any characters normally used for delimiting.
+	const TCHAR* CR_PAIR_DELIM = TEXT("\x01");
+	const TCHAR* CR_PAIR_EQ = TEXT("\x02");
+
+	// Enabled Plugins
+	if (const FXmlNode* EnabledPluginsNode = File.GetRootNode()->FindChildNode(TEXT("EnabledPlugins")))
+	{
+		const TArray<FXmlNode*>& EnabledPlugins = EnabledPluginsNode->GetChildrenNodes();
+		CrashContext.EnabledPluginsOffset = (uint32)(DynamicDataPtr - DynamicDataStart);
+		CrashContext.EnabledPluginsNum = EnabledPlugins.Num();
+		for (const FXmlNode* Node : EnabledPlugins)
+		{
+			FCString::Strncat(DynamicDataPtr, *Node->GetContent(), CR_DYNAMIC_BUFFER_REMAIN);
+			FCString::Strncat(DynamicDataPtr, CR_PAIR_DELIM, CR_DYNAMIC_BUFFER_REMAIN);
+		}
+		DynamicDataPtr += FCString::Strlen(DynamicDataPtr) + 1;
+	}
+	else
+	{
+		CrashContext.EnabledPluginsOffset = (uint32)(DynamicDataPtr - DynamicDataStart);
+		CrashContext.EnabledPluginsNum = 0;
+	}
+	
+	// Engine Data
+	if (const FXmlNode* EngineDataNode = File.GetRootNode()->FindChildNode(TEXT("EngineData")))
+	{
+		const TArray<FXmlNode*>& EngineDataEntries = EngineDataNode->GetChildrenNodes();
+		CrashContext.EngineDataOffset = (uint32)(DynamicDataPtr - DynamicDataStart);
+		CrashContext.EngineDataNum = EngineDataEntries.Num();
+		for (const FXmlNode* Node : EngineDataEntries)
+		{
+			FCString::Strncat(DynamicDataPtr, *Node->GetTag(), CR_DYNAMIC_BUFFER_REMAIN);
+			FCString::Strncat(DynamicDataPtr, CR_PAIR_EQ, CR_DYNAMIC_BUFFER_REMAIN);
+			FCString::Strncat(DynamicDataPtr, *Node->GetContent(), CR_DYNAMIC_BUFFER_REMAIN);
+			FCString::Strncat(DynamicDataPtr, CR_PAIR_DELIM, CR_DYNAMIC_BUFFER_REMAIN);
+		}
+		DynamicDataPtr += FCString::Strlen(DynamicDataPtr) + 1;
+	}
+	else
+	{
+		CrashContext.EngineDataOffset = (uint32)(DynamicDataPtr - DynamicDataStart);
+		CrashContext.EngineDataNum = 0;
+	}
+	
+	// Game Data
+	if (const FXmlNode* GameDataNode = File.GetRootNode()->FindChildNode(TEXT("GameData")))
+	{
+		const TArray<FXmlNode*>& GameDataEntries = GameDataNode->GetChildrenNodes();
+		CrashContext.GameDataOffset = (uint32)(DynamicDataPtr - DynamicDataStart);
+		CrashContext.GameDataNum = GameDataEntries.Num();
+		for (const FXmlNode* Node : GameDataEntries)
+		{
+			FCString::Strncat(DynamicDataPtr, *Node->GetTag(), CR_DYNAMIC_BUFFER_REMAIN);
+			FCString::Strncat(DynamicDataPtr, CR_PAIR_EQ, CR_DYNAMIC_BUFFER_REMAIN);
+			FCString::Strncat(DynamicDataPtr, *Node->GetContent(), CR_DYNAMIC_BUFFER_REMAIN);
+			FCString::Strncat(DynamicDataPtr, CR_PAIR_DELIM, CR_DYNAMIC_BUFFER_REMAIN);
+		}
+		DynamicDataPtr += FCString::Strlen(DynamicDataPtr) + 1;
+	}
+	else
+	{
+		CrashContext.GameDataOffset = (uint32)(DynamicDataPtr - DynamicDataStart);
+		CrashContext.GameDataNum = 0;
+	}
+
+	#undef CR_DYNAMIC_BUFFER_REMAIN
+
 	return true;
 }
 
