Fix crash report client UI to hide screenshot elements properly when no screenshot is present.

diff --git a/Engine/Source/Programs/CrashReportClient/Private/SCrashReportClient.cpp b/Engine/Source/Programs/CrashReportClient/Private/SCrashReportClient.cpp
--- a/Engine/Source/Programs/CrashReportClient/Private/SCrashReportClient.cpp
+++ b/Engine/Source/Programs/CrashReportClient/Private/SCrashReportClient.cpp
@@ -138,7 +138,7 @@ void SCrashReportClient::ConstructDetailedDialog(const TSharedRef<FCrashReportCl
 	}
 	else
 	{
-		ScreenshotSlot.MaxHeight(0);
+		ScreenshotSlot.AutoHeight();
 	}
 
 	ChildSlot
@@ -452,8 +452,8 @@ void SCrashReportClient::ConstructSimpleDialog(const TSharedRef<FCrashReportClie
 	}
 	else
 	{
-		ScreenshotDescriptionSlot.MaxHeight(0);
-		ScreenshotSlot.MaxHeight(0);
+		ScreenshotDescriptionSlot.AutoHeight();
+		ScreenshotSlot.AutoHeight();
 	}
 
 	ChildSlot
