Fix CRC not exiting properly

#rnx

diff --git a/Engine/Source/Programs/CrashReportClient/Private/CrashReportClientApp.cpp b/Engine/Source/Programs/CrashReportClient/Private/CrashReportClientApp.cpp
--- a/Engine/Source/Programs/CrashReportClient/Private/CrashReportClientApp.cpp
+++ b/Engine/Source/Programs/CrashReportClient/Private/CrashReportClientApp.cpp
@@ -1177,9 +1177,7 @@ static void HandleAbnormalShutdown(FSharedCrashContext& CrashContext, uint64 Pro
 
 void RunCrashReportClient(const TCHAR* CommandLine)
 {
-#if !PLATFORM_MAC
 	FTaskTagScope ThreadScope(ETaskTag::EGameThread); // Main thread is the game thread.
-#endif
 
 #if !(UE_BUILD_SHIPPING)
 
