Improve crash report dialog for screenshots during GPU hangs

- Make sure buttons at the bottom are always visible.

diff --git a/Engine/Source/Programs/CrashReportClient/Private/SCrashReportClient.cpp b/Engine/Source/Programs/CrashReportClient/Private/SCrashReportClient.cpp
--- a/Engine/Source/Programs/CrashReportClient/Private/SCrashReportClient.cpp
+++ b/Engine/Source/Programs/CrashReportClient/Private/SCrashReportClient.cpp
@@ -441,8 +441,6 @@ void SCrashReportClient::ConstructSimpleDialog(const TSharedRef<FCrashReportClie
 
 		ScreenshotSlot
 		.Padding(FMargin(4, 4))
-		.MaxHeight(500)
-		.AutoHeight()
 		[
 			SNew(SScaleBox)
 			.Stretch(EStretch::ScaleToFit)
@@ -495,6 +493,7 @@ void SCrashReportClient::ConstructSimpleDialog(const TSharedRef<FCrashReportClie
 
 			// Stuff anchored to the bottom
 			+SVerticalBox::Slot()
+			.AutoHeight()
 			.Padding( FMargin(4, 4, 4, 4) )
 			[
 				SNew(SHorizontalBox)
