[CrashReportClient] Update and break out manual bug report url

diff --git a/Engine/Source/Programs/CrashReportClient/Private/SCrashReportClient.cpp b/Engine/Source/Programs/CrashReportClient/Private/SCrashReportClient.cpp
--- a/Engine/Source/Programs/CrashReportClient/Private/SCrashReportClient.cpp
+++ b/Engine/Source/Programs/CrashReportClient/Private/SCrashReportClient.cpp
@@ -50,7 +50,12 @@ void SCrashReportClient::Construct(const FArguments& InArgs, const TSharedRef<FC
 	bIsUserCommentEmpty = true;
 	bHideSubmitAndRestart = InArgs._bHideSubmitAndRestart;
 
-	FText CrashDetailedMessage = LOCTEXT("CrashDetailed", "We are very sorry that this crash occurred. Our goal is to prevent crashes like this from occurring in the future. Please help us track down and fix this crash by providing detailed information about what you were doing so that we may reproduce the crash and fix it quickly. You can also log a Bug Report with us using the <a id=\"browser\" href=\"https://epicsupport.force.com/unrealengine/s/\" style=\"Hyperlink\">Bug Submission Form</> and work directly with support staff to report this issue.\n\nThanks for your help in improving the Unreal Engine.");
+	// Fallback if the runtime has not set detailed message
+	FText BugSubmissionUrl = FText::FromString(TEXT("https://unrealbugsubmissions.epicgames.com/"));
+	FText CrashDetailedMessage = LOCGEN_FORMAT_ORDERED(
+		LOCTEXT("CrashDetailed", "We are very sorry that this crash occurred. Our goal is to prevent crashes like this from occurring in the future. Please help us track down and fix this crash by providing detailed information about what you were doing so that we may reproduce the crash and fix it quickly. You can also log a Bug Report with us using the <a id=\"browser\" href=\"{0}\" style=\"Hyperlink\">Bug Submission Form</> and work directly with support staff to report this issue.\n\nThanks for your help in improving the Unreal Engine."),
+		BugSubmissionUrl
+	);
 	if (FPrimaryCrashProperties::Get()->IsValid())
 	{
 		FString CrashDetailedMessageString = FPrimaryCrashProperties::Get()->CrashReporterMessage.AsString();
