Fixed mac compile errors

#Synced-CL 50144691

diff --git a/Engine/Source/Programs/CrashReportClient/CrashReportClient.Build.cs b/Engine/Source/Programs/CrashReportClient/CrashReportClient.Build.cs
--- a/Engine/Source/Programs/CrashReportClient/CrashReportClient.Build.cs
+++ b/Engine/Source/Programs/CrashReportClient/CrashReportClient.Build.cs
@@ -6,6 +6,8 @@ public class CrashReportClient : ModuleRules
 {
 	public CrashReportClient(ReadOnlyTargetRules Target) : base(Target)
 	{
+		bRequiresPlatformSDK = true;
+
 		PublicIncludePathModuleNames.Add("Launch");
 
 		PrivateDependencyModuleNames.AddRange(
diff --git a/Engine/Source/Runtime/CrashReportCore/CrashReportCore.Build.cs b/Engine/Source/Runtime/CrashReportCore/CrashReportCore.Build.cs
--- a/Engine/Source/Runtime/CrashReportCore/CrashReportCore.Build.cs
+++ b/Engine/Source/Runtime/CrashReportCore/CrashReportCore.Build.cs
@@ -6,6 +6,8 @@ public class CrashReportCore : ModuleRules
 {
 	public CrashReportCore( ReadOnlyTargetRules Target ) : base(Target)
 	{
+		bRequiresPlatformSDK = true;
+
 		PublicDependencyModuleNames.AddRange(
 			new string[] {
 				"Core",
diff --git a/Engine/Source/Runtime/CrashReportCore/Private/Mac/MacErrorReport.cpp b/Engine/Source/Runtime/CrashReportCore/Private/Mac/MacErrorReport.cpp
--- a/Engine/Source/Runtime/CrashReportCore/Private/Mac/MacErrorReport.cpp
+++ b/Engine/Source/Runtime/CrashReportCore/Private/Mac/MacErrorReport.cpp
@@ -8,6 +8,7 @@
 #include "CrashReportCoreConfig.h"
 
 #if WITH_CRASHREPORTER
+#include "Mac/MacSystemIncludes.h"
 #include "CrashDebugHelper.h"
 #include "CrashDebugHelperModule.h"
 #endif
